removed some platform special cases for platform input.

#rnx

diff --git a/Engine/Plugins/Experimental/CommonUI/Source/CommonInput/CommonInput.Build.cs b/Engine/Plugins/Experimental/CommonUI/Source/CommonInput/CommonInput.Build.cs
--- a/Engine/Plugins/Experimental/CommonUI/Source/CommonInput/CommonInput.Build.cs
+++ b/Engine/Plugins/Experimental/CommonUI/Source/CommonInput/CommonInput.Build.cs
@@ -48,9 +48,6 @@ public class CommonInput : ModuleRules
 			PublicDefinitions.Add("UE_COMMONINPUT_PLATFORM_TYPE = " + CommonUIPlatform);
 		}
 
-		PrivateDefinitions.Add("UE_COMMONINPUT_PLATFORM_USE_GAMEPAD_IF_MOUSE_REMOVED=" + (PlatformUsesGamepadIfMouseRemoved ? "1" : "0") );
-		PrivateDefinitions.Add("UE_COMMONINPUT_PLATFORM_KBM_REQUIRES_ATTACHED_MOUSE=" + (PlatformKBMRequiresAttachedMouse ? "1" : "0") );
-
 		PrivateDependencyModuleNames.Add("GeForceNOWWrapper");
 	}
 
@@ -79,8 +76,4 @@ public class CommonInput : ModuleRules
 
 		return string.Empty;
 	}
-
-	protected virtual bool PlatformUsesGamepadIfMouseRemoved { get { return false; } }
-	protected virtual bool PlatformKBMRequiresAttachedMouse { get { return false; } }
-
 }
diff --git a/Engine/Plugins/Experimental/CommonUI/Source/CommonInput/Private/CommonInputSubsystem.cpp b/Engine/Plugins/Experimental/CommonUI/Source/CommonInput/Private/CommonInputSubsystem.cpp
--- a/Engine/Plugins/Experimental/CommonUI/Source/CommonInput/Private/CommonInputSubsystem.cpp
+++ b/Engine/Plugins/Experimental/CommonUI/Source/CommonInput/Private/CommonInputSubsystem.cpp
@@ -46,13 +46,6 @@ public:
 	
 	virtual void Tick(const float DeltaTime, FSlateApplication& SlateApp, TSharedRef<ICursor> Cursor) override
 	{
-#if UE_COMMONINPUT_PLATFORM_USE_GAMEPAD_IF_MOUSE_REMOVED
-		//fall back to gamepad controls if the mouse is removed... seems to stop slate getting confused and preventing gamepad input
-		if( InputSubsystem.GetCurrentInputType() == ECommonInputType::MouseAndKeyboard && !FSlateApplication::Get().IsMouseAttached() )
-		{
-			InputSubsystem.SetCurrentInputType( ECommonInputType::Gamepad );
-		}
-#endif
 	}
 
 	virtual bool HandleKeyDownEvent(FSlateApplication& SlateApp, const FKeyEvent& InKeyEvent) override
@@ -686,7 +679,7 @@ bool UCommonInputSubsystem::PlatformSupportsInputType(ECommonInputType InInputTy
 	{
 	case ECommonInputType::MouseAndKeyboard:
 	{
-#if UE_COMMONINPUT_PLATFORM_KBM_REQUIRES_ATTACHED_MOUSE || PLATFORM_SWITCH
+#if PLATFORM_SWITCH
 		bPlatformSupportsInput &= FSlateApplication::Get().IsMouseAttached();
 #endif
 	}
