CommonUI - The CommonUI ActionRouterBase shouldn't assume Escape is completely forbidden.  It now also checks to see if its assumption is correct, that the user has left Escape as the key to stop PIE.

diff --git a/Engine/Plugins/Experimental/CommonUI/Source/CommonUI/Private/Input/CommonUIActionRouterBase.cpp b/Engine/Plugins/Experimental/CommonUI/Source/CommonUI/Private/Input/CommonUIActionRouterBase.cpp
--- a/Engine/Plugins/Experimental/CommonUI/Source/CommonUI/Private/Input/CommonUIActionRouterBase.cpp
+++ b/Engine/Plugins/Experimental/CommonUI/Source/CommonUI/Private/Input/CommonUIActionRouterBase.cpp
@@ -19,6 +19,7 @@
 #include "CommonUIUtils.h"
 #include "CommonUISubsystemBase.h"
 #include "Slate/SGameLayerManager.h"
+#include "Framework/Commands/InputBindingManager.h"
 
 bool bAlwaysShowCursor = false;
 static const FAutoConsoleVariableRef CVarAlwaysShowCursor(
@@ -289,10 +290,20 @@ ERouteUIInputResult UCommonUIActionRouterBase::ProcessInput(FKey Key, EInputEven
 	// In PIE, let unmodified escape through (people expect it to close PIE)
 	if (GIsPlayInEditorWorld && InputEvent == IE_Pressed && Key == EKeys::Escape)
 	{
-		FModifierKeysState ModifierKeys = FSlateApplication::Get().GetModifierKeys();
-		if (!ModifierKeys.IsAltDown() && !ModifierKeys.IsCommandDown() && !ModifierKeys.IsControlDown() && !ModifierKeys.IsShiftDown())
+		TSharedPtr<FUICommandInfo> StopCommand = FInputBindingManager::Get().FindCommandInContext("PlayWorld", "StopPlaySession");
+		if (ensure(StopCommand))
 		{
-			return ERouteUIInputResult::Unhandled;
+			// If the user has rebound the 'Stop' to something other than vanilla Escape, then don't do this,
+			// they're clearly trying to avoid collisions with the UI Actions.
+			const TSharedRef<const FInputChord> PrimaryCord = StopCommand->GetActiveChord(EMultipleKeyBindingIndex::Primary);
+			if (PrimaryCord->Key == EKeys::Escape && PrimaryCord->bAlt == false && PrimaryCord->bCmd == false && PrimaryCord->bCtrl == false && PrimaryCord->bShift == false)
+			{
+				FModifierKeysState ModifierKeys = FSlateApplication::Get().GetModifierKeys();
+				if (!ModifierKeys.IsAltDown() && !ModifierKeys.IsCommandDown() && !ModifierKeys.IsControlDown() && !ModifierKeys.IsShiftDown())
+				{
+					return ERouteUIInputResult::Unhandled;
+				}
+			}
 		}
 	}
 #endif
