CommonUI - The AnalogSlider shoud just inherit the slider style, not use the default slate one.

diff --git a/Engine/Plugins/Experimental/CommonUI/Source/CommonUI/Private/AnalogSlider.cpp b/Engine/Plugins/Experimental/CommonUI/Source/CommonUI/Private/AnalogSlider.cpp
--- a/Engine/Plugins/Experimental/CommonUI/Source/CommonUI/Private/AnalogSlider.cpp
+++ b/Engine/Plugins/Experimental/CommonUI/Source/CommonUI/Private/AnalogSlider.cpp
@@ -16,8 +16,6 @@ UAnalogSlider::UAnalogSlider(const FObjectInitializer& ObjectInitializer)
 	SliderBarColor = FLinearColor::White;
 	SliderHandleColor = FLinearColor::White;
 	StepSize = 0.01f;
-	SSlider::FArguments Defaults;
-	WidgetStyle = *Defaults._Style;
 	IsFocusable = true;
 }
 
