CommonUI - marking DefaultRichTextDataClass, allow abstract false, and the base class abstract.

diff --git a/Engine/Plugins/Experimental/CommonUI/Source/CommonUI/Public/CommonUIRichTextData.h b/Engine/Plugins/Experimental/CommonUI/Source/CommonUI/Public/CommonUIRichTextData.h
--- a/Engine/Plugins/Experimental/CommonUI/Source/CommonUI/Public/CommonUIRichTextData.h
+++ b/Engine/Plugins/Experimental/CommonUI/Source/CommonUI/Public/CommonUIRichTextData.h
@@ -26,7 +26,7 @@ struct COMMONUI_API FRichTextIconData : public FTableRowBase
 /* Derive from this class for rich text data per game
  * it is referenced in Common UI Settings, found in project settings UI
  */
-UCLASS(Blueprintable, meta = (Category = "Common UI"))
+UCLASS(Abstract, Blueprintable, meta = (Category = "Common UI"))
 class COMMONUI_API UCommonUIRichTextData : public UObject
 {
 	GENERATED_BODY()
diff --git a/Engine/Plugins/Experimental/CommonUI/Source/CommonUI/Public/CommonUISettings.h b/Engine/Plugins/Experimental/CommonUI/Source/CommonUI/Public/CommonUISettings.h
--- a/Engine/Plugins/Experimental/CommonUI/Source/CommonUI/Public/CommonUISettings.h
+++ b/Engine/Plugins/Experimental/CommonUI/Source/CommonUI/Public/CommonUISettings.h
@@ -54,7 +54,7 @@ private:
 	TSoftObjectPtr<UMaterialInterface> DefaultThrobberMaterial;
 
 	/** The Default Data for rich text to show inline icon and others. */
-	UPROPERTY(config, EditAnywhere, Category = "RichText")
+	UPROPERTY(config, EditAnywhere, Category = "RichText", meta=(AllowAbstract=false))
 	TSoftClassPtr<UCommonUIRichTextData> DefaultRichTextDataClass;
 
 private:
