CommonInput - Using the new UPlatformSettings, to store per-platform information into an NDA location when each platform is configured.  These changes make it so we use the 'real' platform ini name, e.g. There's no more "PC" platform as far as the common input exists for configuring the platforms.  So if you setup linux, that's another platform you'd configure.  The gamepads are still wonky, they're trying to use platform ini names for gamepad names, going to refactor that so it's not coming from platform names.  Made an upgrade path for the old settings that will update and then clear the old configuration.  Also removed some settings that were added to DataDrivenPlatformInfo that don't make sense, and updated some platforms that never got the settings to begin with.

diff --git a/Engine/Config/Android/DataDrivenPlatformInfo.ini b/Engine/Config/Android/DataDrivenPlatformInfo.ini
--- a/Engine/Config/Android/DataDrivenPlatformInfo.ini
+++ b/Engine/Config/Android/DataDrivenPlatformInfo.ini
@@ -1,12 +1,7 @@
 [DataDrivenPlatformInfo]
 bIsConfidential=false
 AudioCompressionSettingsIniSectionName=/Script/AndroidRuntimeSettings.AndroidRuntimeSettings
-bIsInteractablePlatform=true
 bHasDedicatedGamepad=false
-bIsGenericPC=false
-bInputSupportConfigurable=falbIsInteractablePlatform=true
-bHasDedicatedGamepad=false
-bIsGenericPC=false
 bInputSupportConfigurable=false
 DefaultInputType=Touch
 bSupportsMouseAndKeyboard=false
diff --git a/Engine/Config/HoloLens/DataDrivenPlatformInfo.ini b/Engine/Config/HoloLens/DataDrivenPlatformInfo.ini
--- a/Engine/Config/HoloLens/DataDrivenPlatformInfo.ini
+++ b/Engine/Config/HoloLens/DataDrivenPlatformInfo.ini
@@ -2,6 +2,14 @@
 bIsConfidential=false
 Freezing_MaxFieldAlignment64=8
 Freezing_bAlignBases=true
+bHasDedicatedGamepad=false
+bDefaultInputStandardKeyboard=false
+bInputSupportConfigurable=true
+DefaultInputType=Gamepad
+bSupportsMouseAndKeyboard=true
+bSupportsGamepad=true
+bCanChangeGamepadType=true
+bSupportsTouch=false
 GlobalIdentifier=0DF604E112E4445280BEE1C74EB934B1
 ; Oodle is not supported on this platform so force something else to be used (Zlib or None) :
 HardwareCompressionFormat=Zlib
diff --git a/Engine/Config/IOS/DataDrivenPlatformInfo.ini b/Engine/Config/IOS/DataDrivenPlatformInfo.ini
--- a/Engine/Config/IOS/DataDrivenPlatformInfo.ini
+++ b/Engine/Config/IOS/DataDrivenPlatformInfo.ini
@@ -1,9 +1,7 @@
 [DataDrivenPlatformInfo]
 bIsConfidential=false
 AudioCompressionSettingsIniSectionName=/Script/IOSRuntimeSettings.IOSRuntimeSettings
-bIsInteractablePlatform=true
 bHasDedicatedGamepad=false
-bIsGenericPC=false
 bInputSupportConfigurable=false
 DefaultInputType=Touch
 bSupportsMouseAndKeyboard=false
diff --git a/Engine/Config/Linux/DataDrivenPlatformInfo.ini b/Engine/Config/Linux/DataDrivenPlatformInfo.ini
--- a/Engine/Config/Linux/DataDrivenPlatformInfo.ini
+++ b/Engine/Config/Linux/DataDrivenPlatformInfo.ini
@@ -1,6 +1,5 @@
 [DataDrivenPlatformInfo]
 IniParent=Unix
-bIsInteractablePlatform=true
 bHasDedicatedGamepad=false
 bDefaultInputStandardKeyboard=true
 bInputSupportConfigurable=true
diff --git a/Engine/Config/LinuxAArch64/DataDrivenPlatformInfo.ini b/Engine/Config/LinuxAArch64/DataDrivenPlatformInfo.ini
--- a/Engine/Config/LinuxAArch64/DataDrivenPlatformInfo.ini
+++ b/Engine/Config/LinuxAArch64/DataDrivenPlatformInfo.ini
@@ -1,7 +1,6 @@
 [DataDrivenPlatformInfo]
 IniParent=Unix
 bIsConfidential=false
-bIsInteractablePlatform=true
 bHasDedicatedGamepad=false
 bDefaultInputStandardKeyboard=true
 bInputSupportConfigurable=true
diff --git a/Engine/Config/Lumin/DataDrivenPlatformInfo.ini b/Engine/Config/Lumin/DataDrivenPlatformInfo.ini
--- a/Engine/Config/Lumin/DataDrivenPlatformInfo.ini
+++ b/Engine/Config/Lumin/DataDrivenPlatformInfo.ini
@@ -1,6 +1,14 @@
 [DataDrivenPlatformInfo]
 bIsConfidential=false
 IniParent=Android
+bHasDedicatedGamepad=false
+bDefaultInputStandardKeyboard=true
+bInputSupportConfigurable=true
+DefaultInputType=Gamepad
+bSupportsMouseAndKeyboard=true
+bSupportsGamepad=true
+bCanChangeGamepadType=true
+bSupportsTouch=false
 GlobalIdentifier=A20FCE56EA274AADB4EA8567AA6BBAB3
 
 ; Oodle is not supported on this platform so force something else to be used (Zlib or None) :
diff --git a/Engine/Config/Mac/DataDrivenPlatformInfo.ini b/Engine/Config/Mac/DataDrivenPlatformInfo.ini
--- a/Engine/Config/Mac/DataDrivenPlatformInfo.ini
+++ b/Engine/Config/Mac/DataDrivenPlatformInfo.ini
@@ -1,8 +1,8 @@
 [DataDrivenPlatformInfo]
 bIsConfidential=false
-AudioCompressionSettingsIniSectionName=/Script/MacTargetPlatform.MacTargetSettingsbIsInteractablePlatform=true
+AudioCompressionSettingsIniSectionName=/Script/MacTargetPlatform.MacTargetSettings
 bHasDedicatedGamepad=false
-bIsGenericPC=false
+bDefaultInputStandardKeyboard=false
 bInputSupportConfigurable=true
 DefaultInputType=MouseAndKeyboard
 bSupportsMouseAndKeyboard=true
diff --git a/Engine/Config/Windows/DataDrivenPlatformInfo.ini b/Engine/Config/Windows/DataDrivenPlatformInfo.ini
--- a/Engine/Config/Windows/DataDrivenPlatformInfo.ini
+++ b/Engine/Config/Windows/DataDrivenPlatformInfo.ini
@@ -6,7 +6,6 @@ Freezing_bWithRayTracing=true
 Freezing_MaxFieldAlignment32=4
 Freezing_MaxFieldAlignment64=8
 Freezing_bAlignBases=true
-bIsInteractablePlatform=true
 bHasDedicatedGamepad=false
 bDefaultInputStandardKeyboard=true
 bInputSupportConfigurable=true
diff --git a/Engine/Plugins/Experimental/CommonUI/Source/CommonInput/CommonInput.Build.cs b/Engine/Plugins/Experimental/CommonUI/Source/CommonInput/CommonInput.Build.cs
--- a/Engine/Plugins/Experimental/CommonUI/Source/CommonInput/CommonInput.Build.cs
+++ b/Engine/Plugins/Experimental/CommonUI/Source/CommonInput/CommonInput.Build.cs
@@ -8,11 +8,12 @@ public class CommonInput : ModuleRules
 	{
 		PublicDependencyModuleNames.AddRange(
 			new string[]
-			{
+            {
                 "Core",
                 "CoreUObject",
                 "Engine",
-                "InputCore"
+                "InputCore",
+				"DeveloperSettings"
             }
 		);
 
diff --git a/Engine/Plugins/Experimental/CommonUI/Source/CommonInput/Private/CommonInputBaseTypes.cpp b/Engine/Plugins/Experimental/CommonUI/Source/CommonInput/Private/CommonInputBaseTypes.cpp
--- a/Engine/Plugins/Experimental/CommonUI/Source/CommonInput/Private/CommonInputBaseTypes.cpp
+++ b/Engine/Plugins/Experimental/CommonUI/Source/CommonInput/Private/CommonInputBaseTypes.cpp
@@ -8,6 +8,7 @@
 #include "ICommonInputModule.h"
 #include "Misc/DataDrivenPlatformInfoRegistry.h"
 #include "UObject/ObjectSaveContext.h"
+#include "CommonInputBaseTypes.h"
 
 const FName FCommonInputDefaults::PlatformPC(TEXT("PC"));
 const FName FCommonInputDefaults::GamepadGeneric(TEXT("Generic"));
@@ -130,7 +131,22 @@ const TArray<FName>& UCommonInputBaseControllerData::GetRegisteredGamepads()
 			const FName PlatformName = Platform.Key;
 			const FDataDrivenPlatformInfo& PlatformInfo = Platform.Value;
 
-			if (!PlatformInfo.bDefaultInputStandardKeyboard && PlatformInfo.bIsInteractablePlatform && PlatformInfo.bHasDedicatedGamepad)
+			// Don't add fake platforms that are used to group real platforms to make configuration for groups of platforms
+			// simpler.
+			if (PlatformInfo.bIsFakePlatform)
+			{
+				continue;
+			}
+
+			// If the platform uses the standard keyboard for default input, ignore it, all of those platforms will use "PC"
+			// as their target, so Windows, Linux, but not Mac.
+			if (PlatformInfo.bDefaultInputStandardKeyboard)
+			{
+				continue;
+			}
+
+			// Only add platforms with dedicated gamepads.
+			if (PlatformInfo.bHasDedicatedGamepad)
 			{
 				RegisteredGamepads.Add(PlatformName);
 			}
@@ -140,6 +156,191 @@ const TArray<FName>& UCommonInputBaseControllerData::GetRegisteredGamepads()
 	static TArray<FName> RegisteredGamepads = GenerateRegisteredGamepads();
 	return RegisteredGamepads;
 }
+
+#if WITH_EDITOR
+void UCommonInputBaseControllerData::PostEditChangeProperty(struct FPropertyChangedEvent& PropertyChangedEvent)
+{
+	Super::PostEditChangeProperty(PropertyChangedEvent);
+
+	if (PropertyChangedEvent.ChangeType == EPropertyChangeType::ValueSet)
+	{
+		if (PropertyChangedEvent.GetPropertyName() == GET_MEMBER_NAME_CHECKED(UCommonInputBaseControllerData, SetButtonImageHeightTo))
+		{
+			if (SetButtonImageHeightTo != 0)
+			{
+				for (FCommonInputKeyBrushConfiguration& BrushConfig : InputBrushDataMap)
+				{
+					FVector2D NewBrushSize = BrushConfig.KeyBrush.GetImageSize();
+					NewBrushSize.X = FMath::RoundToInt(NewBrushSize.X * (SetButtonImageHeightTo / NewBrushSize.X));
+					NewBrushSize.Y = SetButtonImageHeightTo;
+
+					BrushConfig.KeyBrush.SetImageSize(NewBrushSize);
+				}
+
+				for (FCommonInputKeySetBrushConfiguration& BrushConfig : InputBrushKeySets)
+				{
+					FVector2D NewBrushSize = BrushConfig.KeyBrush.GetImageSize();
+					NewBrushSize.X = FMath::RoundToInt(NewBrushSize.X * (SetButtonImageHeightTo / NewBrushSize.X));
+					NewBrushSize.Y = SetButtonImageHeightTo;
+
+					BrushConfig.KeyBrush.SetImageSize(NewBrushSize);
+				}
+			}
+		}
+	}
+}
+#endif
+
+
+
+
+UCommonInputPlatformSettings::UCommonInputPlatformSettings()
+{
+	DefaultInputType = ECommonInputType::Gamepad;
+	bSupportsMouseAndKeyboard = false;
+	bSupportsGamepad = true;
+	bCanChangeGamepadType = true;
+	bSupportsTouch = false;
+	DefaultGamepadName = FCommonInputDefaults::GamepadGeneric;
+}
+
+void UCommonInputPlatformSettings::PostInitProperties()
+{
+	Super::PostInitProperties();
+
+	ControllerDataClasses.Reset();
+	InitializeControllerData();
+}
+
+void UCommonInputPlatformSettings::InitializeControllerData() const
+{
+	if (ControllerData.Num() != ControllerDataClasses.Num())
+	{
+		for (TSoftClassPtr<UCommonInputBaseControllerData> ControllerDataPtr : ControllerData)
+		{
+			if (TSubclassOf<UCommonInputBaseControllerData> ControllerDataClass = ControllerDataPtr.LoadSynchronous())
+			{
+				ControllerDataClasses.Add(ControllerDataClass);
+			}
+		}
+	}
+}
+
+void UCommonInputPlatformSettings::InitializePlatformDefaults()
+{
+	const FName PlatformName = *GetPlatformIniName();
+	const FDataDrivenPlatformInfo& PlatformInfo = FDataDrivenPlatformInfoRegistry::GetPlatformInfo(PlatformName);
+	if (PlatformInfo.DefaultInputType == "Gamepad")
+	{
+		DefaultInputType = ECommonInputType::Gamepad;
+	}
+	else if (PlatformInfo.DefaultInputType == "Touch")
+	{
+		DefaultInputType = ECommonInputType::Touch;
+	}
+	else if (PlatformInfo.DefaultInputType == "MouseAndKeyboard")
+	{
+		DefaultInputType = ECommonInputType::MouseAndKeyboard;
+	}
+
+	bSupportsMouseAndKeyboard = PlatformInfo.bSupportsMouseAndKeyboard;
+	bSupportsGamepad = PlatformInfo.bSupportsGamepad;
+	bCanChangeGamepadType = PlatformInfo.bCanChangeGamepadType;
+	bSupportsTouch = PlatformInfo.bSupportsTouch;
+
+	DefaultGamepadName = PlatformName;
+}
+
+bool UCommonInputPlatformSettings::TryGetInputBrush(FSlateBrush& OutBrush, FKey Key, ECommonInputType InputType, const FName& GamepadName) const
+{
+	InitializeControllerData();
+
+	for (const TSubclassOf<UCommonInputBaseControllerData>& ControllerDataPtr : ControllerDataClasses)
+	{
+		const UCommonInputBaseControllerData* DefaultControllerData = ControllerDataPtr.GetDefaultObject();
+		if (DefaultControllerData && DefaultControllerData->InputType == InputType)
+		{
+			if (DefaultControllerData->InputType != ECommonInputType::Gamepad || DefaultControllerData->GamepadName == GamepadName)
+			{
+				return DefaultControllerData->TryGetInputBrush(OutBrush, Key);
+			}
+		}
+	}
+
+	return false;
+}
+
+bool UCommonInputPlatformSettings::TryGetInputBrush(FSlateBrush& OutBrush, const TArray<FKey>& Keys, ECommonInputType InputType, const FName& GamepadName) const
+{
+	InitializeControllerData();
+
+	for (const TSubclassOf<UCommonInputBaseControllerData>& ControllerDataPtr : ControllerDataClasses)
+	{
+		const UCommonInputBaseControllerData* DefaultControllerData = ControllerDataPtr.GetDefaultObject();
+		if (DefaultControllerData && DefaultControllerData->InputType == InputType)
+		{
+			if (DefaultControllerData->InputType != ECommonInputType::Gamepad || DefaultControllerData->GamepadName == GamepadName)
+			{
+				return DefaultControllerData->TryGetInputBrush(OutBrush, Keys);
+			}
+		}
+	}
+
+	return false;
+}
+
+bool UCommonInputPlatformSettings::SupportsInputType(ECommonInputType InputType) const
+{
+	switch (InputType)
+	{
+	case ECommonInputType::MouseAndKeyboard:
+	{
+		return bSupportsMouseAndKeyboard;
+	}
+	break;
+	case ECommonInputType::Gamepad:
+	{
+		return bSupportsGamepad;
+	}
+	break;
+	case ECommonInputType::Touch:
+	{
+		return bSupportsTouch;
+	}
+	break;
+	}
+	return false;
+}
+
+#if WITH_EDITOR
+void UCommonInputPlatformSettings::PostEditChangeProperty(struct FPropertyChangedEvent& PropertyChangedEvent)
+{
+	ControllerDataClasses.Reset();
+}
+#endif
+
+
+
+
+
+
+
+
+
+
+
+
+
+
+
+
+
+
+
+
+
+
+
 
 bool FCommonInputPlatformBaseData::TryGetInputBrush(FSlateBrush& OutBrush, FKey Key, ECommonInputType InputType, const FName& GamepadName) const
 {
@@ -193,11 +394,23 @@ const TArray<FName>& FCommonInputPlatformBaseData::GetRegisteredPlatforms()
 			const FName PlatformName = Platform.Key;
 			const FDataDrivenPlatformInfo& PlatformInfo = Platform.Value;
 
-			if (!PlatformInfo.bDefaultInputStandardKeyboard && PlatformInfo.bIsInteractablePlatform)
+			// Don't add fake platforms that are used to group real platforms to make configuration for groups of platforms
+			// simpler.
+			if (PlatformInfo.bIsFakePlatform)
 			{
-				RegisteredPlatforms.Add(PlatformName);
+				continue;
 			}
+
+			// If the platform uses the standard keyboard for default input, ignore it, all of those platforms will use "PC"
+			// as their target, so Windows, Linux, but not Mac.
+			if (PlatformInfo.bDefaultInputStandardKeyboard)
+			{
+				continue;
+			}
+
+			RegisteredPlatforms.Add(PlatformName);
 		}
+
 		return RegisteredPlatforms;
 	};
 	static TArray<FName> RegisteredPlatforms = GenerateRegisteredPlatforms();
@@ -220,10 +433,6 @@ UCommonInputSettings* FCommonInputBase::GetInputSettings()
 
 void FCommonInputBase::GetCurrentPlatformDefaults(ECommonInputType& OutDefaultInputType, FName& OutDefaultGamepadName)
 {
-	return ICommonInputModule::GetSettings().GetCurrentPlatformDefaults(OutDefaultInputType, OutDefaultGamepadName);
-}
-
-FCommonInputPlatformBaseData FCommonInputBase::GetCurrentBasePlatformData()
-{
-	return ICommonInputModule::GetSettings().GetCurrentPlatform();
+	OutDefaultInputType = UCommonInputPlatformSettings::Get()->GetDefaultInputType();
+	OutDefaultGamepadName = UCommonInputPlatformSettings::Get()->GetDefaultGamepadName();
 }
diff --git a/Engine/Plugins/Experimental/CommonUI/Source/CommonInput/Private/CommonInputModule.cpp b/Engine/Plugins/Experimental/CommonUI/Source/CommonInput/Private/CommonInputModule.cpp
--- a/Engine/Plugins/Experimental/CommonUI/Source/CommonInput/Private/CommonInputModule.cpp
+++ b/Engine/Plugins/Experimental/CommonUI/Source/CommonInput/Private/CommonInputModule.cpp
@@ -50,14 +50,14 @@ void FCommonInputModule::CreateInputSettings()
 	// Register settings
 	ISettingsModule* SettingsModule = FModuleManager::GetModulePtr<ISettingsModule>("Settings");
 
-	if (SettingsModule != nullptr)
-	{
-		ISettingsSectionPtr SettingsSection = SettingsModule->RegisterSettings("Project", "Plugins", "CommonInput",
-			NSLOCTEXT("CommonInputPlugin", "CommonInputSettingsName", "Common Input"),
-			NSLOCTEXT("CommonInputPlugin", "CommonInputSettingsDescription", "Configure Common Input defaults."),
-			CommonInputSettings
-		);
-	}
+	//if (SettingsModule != nullptr)
+	//{
+	//	ISettingsSectionPtr SettingsSection = SettingsModule->RegisterSettings("Project", "Plugins", "CommonInput",
+	//		NSLOCTEXT("CommonInputPlugin", "CommonInputSettingsName", "Common Input"),
+	//		NSLOCTEXT("CommonInputPlugin", "CommonInputSettingsDescription", "Configure Common Input defaults."),
+	//		CommonInputSettings
+	//	);
+	//}
 
 #endif //WITH_EDITOR
 }
diff --git a/Engine/Plugins/Experimental/CommonUI/Source/CommonInput/Private/CommonInputSettings.cpp b/Engine/Plugins/Experimental/CommonUI/Source/CommonInput/Private/CommonInputSettings.cpp
--- a/Engine/Plugins/Experimental/CommonUI/Source/CommonInput/Private/CommonInputSettings.cpp
+++ b/Engine/Plugins/Experimental/CommonUI/Source/CommonInput/Private/CommonInputSettings.cpp
@@ -11,42 +11,7 @@ UCommonInputSettings::UCommonInputSettings(const FObjectInitializer& Initializer
 	: Super(Initializer)
 	, bInputDataLoaded(false)
 {
-	FCommonInputPlatformBaseData PcPlatformData;
-	PcPlatformData.bSupported = true;
-	PcPlatformData.DefaultInputType = ECommonInputType::MouseAndKeyboard;
-	PcPlatformData.bSupportsMouseAndKeyboard = true;
-	CommonInputPlatformData.Add(FCommonInputDefaults::PlatformPC, PcPlatformData);
-
-	for (const TPair<FName, FDataDrivenPlatformInfo>& Platform : FDataDrivenPlatformInfoRegistry::GetAllPlatformInfos())
-	{
-		const FName PlatformName = Platform.Key;
-		const FDataDrivenPlatformInfo& PlatformInfo = Platform.Value;
-
-		if (!PlatformInfo.bDefaultInputStandardKeyboard && PlatformInfo.bIsInteractablePlatform)
-		{
-			FCommonInputPlatformBaseData PlatformData;
-			PlatformData.bSupported = PlatformInfo.bInputSupportConfigurable;
-			if (PlatformInfo.DefaultInputType == "Gamepad")
-			{
-				PlatformData.DefaultInputType = ECommonInputType::Gamepad;
-			}
-			else if (PlatformInfo.DefaultInputType == "Touch")
-			{
-				PlatformData.DefaultInputType = ECommonInputType::Touch;
-			}
-			else if (PlatformInfo.DefaultInputType == "MouseAndKeyboard")
-			{
-				PlatformData.DefaultInputType = ECommonInputType::MouseAndKeyboard;
-			}
-			PlatformData.bSupportsMouseAndKeyboard = PlatformInfo.bSupportsMouseAndKeyboard;
-			PlatformData.bSupportsGamepad = PlatformInfo.bSupportsGamepad;
-			PlatformData.bCanChangeGamepadType = PlatformInfo.bSupportsTouch;
-			PlatformData.bSupportsTouch = PlatformInfo.bCanChangeGamepadType;
-
-			PlatformData.DefaultGamepadName = PlatformName;
-			CommonInputPlatformData.Add(PlatformName, PlatformData);
-		}
-	}
+	PlatformInput.Settings = UPlatformSettings::GetAllPlatformSettings<UCommonInputPlatformSettings>();
 }
 
 void UCommonInputSettings::LoadData()
@@ -62,11 +27,6 @@ void UCommonInputSettings::PostEditChangeProperty(struct FPropertyChangedEvent&
 }
 #endif
 
-const TArray<FName>& UCommonInputSettings::GetRegisteredPlatforms()
-{
-	return FCommonInputPlatformBaseData::GetRegisteredPlatforms();
-}
-
 void UCommonInputSettings::LoadInputData()
 {
 	if (!bInputDataLoaded)
@@ -84,18 +44,18 @@ void UCommonInputSettings::LoadInputData()
 			}
 		}
 		
-		CurrentPlatform = CommonInputPlatformData[FCommonInputBase::GetCurrentPlatformName()];
-		for (TSoftClassPtr<UCommonInputBaseControllerData> ControllerData : CurrentPlatform.GetControllerData())
-		{
-			if (TSubclassOf<UCommonInputBaseControllerData> ControllerDataClass = ControllerData.LoadSynchronous())
-			{
-				CurrentPlatform.ControllerDataClasses.Add(ControllerDataClass);
-				if (bIsDisregardForGC)
-				{
-					ControllerDataClass->AddToRoot();
-				}
-			}
-		}
+		//CurrentPlatform = CommonInputPlatformData[FCommonInputBase::GetCurrentPlatformName()];
+		//for (TSoftClassPtr<UCommonInputBaseControllerData> ControllerData : CurrentPlatform.GetControllerData())
+		//{
+		//	if (TSubclassOf<UCommonInputBaseControllerData> ControllerDataClass = ControllerData.LoadSynchronous())
+		//	{
+		//		CurrentPlatform.ControllerDataClasses.Add(ControllerDataClass);
+		//		if (bIsDisregardForGC)
+		//		{
+		//			ControllerDataClass->AddToRoot();
+		//		}
+		//	}
+		//}
 		bInputDataLoaded = true;
 	}
 }
@@ -103,15 +63,15 @@ void UCommonInputSettings::LoadInputData()
 void UCommonInputSettings::ValidateData()
 {
     bInputDataLoaded &= !InputData.IsPending();
-    for (TSoftClassPtr<UCommonInputBaseControllerData> ControllerData : CurrentPlatform.GetControllerData())
-    {
-		bInputDataLoaded &= CurrentPlatform.ControllerDataClasses.ContainsByPredicate([&ControllerData](const TSubclassOf<UCommonInputBaseControllerData>& ControllerDataClass)
-			{
-				return ControllerDataClass.Get() == ControllerData.Get();
-			});
-
-        bInputDataLoaded &= !ControllerData.IsPending();
-    }
+  //  for (TSoftClassPtr<UCommonInputBaseControllerData> ControllerData : CurrentPlatform.GetControllerData())
+  //  {
+		//bInputDataLoaded &= CurrentPlatform.ControllerDataClasses.ContainsByPredicate([&ControllerData](const TSubclassOf<UCommonInputBaseControllerData>& ControllerDataClass)
+		//	{
+		//		return ControllerDataClass.Get() == ControllerData.Get();
+		//	});
+
+  //      bInputDataLoaded &= !ControllerData.IsPending();
+  //  }
  
 #if !WITH_EDITOR
     UE_CLOG(!bInputDataLoaded, LogCommonInput, Warning, TEXT("Trying to access unloaded CommmonInputSettings data. This may force a sync load."));
@@ -148,18 +108,47 @@ FDataTableRowHandle UCommonInputSettings::GetDefaultBackAction() const
 	return FDataTableRowHandle();
 }
 
-void UCommonInputSettings::GetCurrentPlatformDefaults(ECommonInputType& OutDefaultInputType, FName& OutDefaultGamepadName) const
+void UCommonInputSettings::PostInitProperties()
 {
-	// Defaults can be accessed before platform data is fully loaded, so access them from the array rather than the cached CurrentPlatform.
-	const FCommonInputPlatformBaseData& CurrentPlatformData = CommonInputPlatformData[FCommonInputBase::GetCurrentPlatformName()];
-	OutDefaultInputType = CurrentPlatformData.GetDefaultInputType();
-	OutDefaultGamepadName = CurrentPlatformData.GetDefaultGamepadName();
-}
+	Super::PostInitProperties();
 
-FCommonInputPlatformBaseData UCommonInputSettings::GetCurrentPlatform() const
-{
-    const_cast<UCommonInputSettings*>(this)->ValidateData();
-    
-	ensure(bInputDataLoaded);
-	return CurrentPlatform;
+	if (CommonInputPlatformData_DEPRECATED.Num())
+	{
+		for (const auto& PlatformData : CommonInputPlatformData_DEPRECATED)
+		{
+			const FCommonInputPlatformBaseData& OriginalData = PlatformData.Value;
+
+			if (UCommonInputPlatformSettings* Settings = UPlatformSettings::GetSettingsForPlatform<UCommonInputPlatformSettings>(PlatformData.Key.ToString()))
+			{
+				Settings->bSupportsMouseAndKeyboard = OriginalData.bSupportsMouseAndKeyboard;
+				Settings->bSupportsGamepad = OriginalData.bSupportsGamepad;
+				Settings->bSupportsTouch = OriginalData.bSupportsTouch;
+				Settings->bCanChangeGamepadType = OriginalData.bCanChangeGamepadType;
+				Settings->DefaultGamepadName = OriginalData.DefaultGamepadName;
+				Settings->DefaultInputType = OriginalData.DefaultInputType;
+				Settings->ControllerData = OriginalData.ControllerData;
+			}
+			else if (PlatformData.Key == FCommonInputDefaults::PlatformPC)
+			{
+				TArray<UCommonInputPlatformSettings*> PCPlatforms;
+				PCPlatforms.Add(UPlatformSettings::GetSettingsForPlatform<UCommonInputPlatformSettings>("Windows"));
+				PCPlatforms.Add(UPlatformSettings::GetSettingsForPlatform<UCommonInputPlatformSettings>("WinGDK"));
+				PCPlatforms.Add(UPlatformSettings::GetSettingsForPlatform<UCommonInputPlatformSettings>("Linux"));
+
+				for (UCommonInputPlatformSettings* PCPlatform : PCPlatforms)
+				{
+					PCPlatform->bSupportsMouseAndKeyboard = OriginalData.bSupportsMouseAndKeyboard;
+					PCPlatform->bSupportsGamepad = OriginalData.bSupportsGamepad;
+					PCPlatform->bSupportsTouch = OriginalData.bSupportsTouch;
+					PCPlatform->bCanChangeGamepadType = OriginalData.bCanChangeGamepadType;
+					PCPlatform->DefaultGamepadName = OriginalData.DefaultGamepadName;
+					PCPlatform->DefaultInputType = OriginalData.DefaultInputType;
+					PCPlatform->ControllerData = OriginalData.ControllerData;
+				}
+			}
+		}
+
+		CommonInputPlatformData_DEPRECATED.Reset();
+		UpdateDefaultConfigFile();
+	}
 }
diff --git a/Engine/Plugins/Experimental/CommonUI/Source/CommonInput/Private/CommonInputSubsystem.cpp b/Engine/Plugins/Experimental/CommonUI/Source/CommonInput/Private/CommonInputSubsystem.cpp
--- a/Engine/Plugins/Experimental/CommonUI/Source/CommonInput/Private/CommonInputSubsystem.cpp
+++ b/Engine/Plugins/Experimental/CommonUI/Source/CommonInput/Private/CommonInputSubsystem.cpp
@@ -284,8 +284,8 @@ void UCommonInputSubsystem::Initialize(FSubsystemCollectionBase& Collection)
 	Super::Initialize(Collection);
 
 	FCommonInputBase::GetInputSettings()->LoadData();
-	GamepadInputType = FCommonInputBase::GetCurrentBasePlatformData().GetDefaultGamepadName();
-	CurrentInputType = LastInputType = FCommonInputBase::GetCurrentBasePlatformData().GetDefaultInputType();
+	GamepadInputType = UPlatformSettings::GetSettingsForPlatform<UCommonInputPlatformSettings>()->GetDefaultGamepadName();
+	CurrentInputType = UPlatformSettings::GetSettingsForPlatform<UCommonInputPlatformSettings>()->GetDefaultInputType();
 
 	CommonInputPreprocessor = MakeShared<FCommonInputPreprocessor>(*this);
 	FSlateApplication::Get().RegisterInputPreProcessor(CommonInputPreprocessor, 0);
@@ -375,7 +375,7 @@ ECommonInputType UCommonInputSubsystem::GetCurrentInputType() const
 
 ECommonInputType UCommonInputSubsystem::GetDefaultInputType() const
 {
-	return FCommonInputBase::GetCurrentBasePlatformData().GetDefaultInputType();
+	return UCommonInputPlatformSettings::Get()->GetDefaultInputType();
 }
 
 void UCommonInputSubsystem::BroadcastInputMethodChanged()
@@ -498,7 +498,7 @@ const FName UCommonInputSubsystem::GetCurrentGamepadName() const
  
 void UCommonInputSubsystem::SetGamepadInputType(const FName& InGamepadInputType)
 {
-	if (!FCommonInputBase::GetCurrentBasePlatformData().CanChangeGamepadType())
+	if (!UCommonInputPlatformSettings::Get()->CanChangeGamepadType())
 	{
 		ensure(false);	// should not be called on the console platforms
 		return;
@@ -674,7 +674,7 @@ FVector2D UCommonInputSubsystem::ClampPositionToViewport(const FVector2D& InPosi
 
 bool UCommonInputSubsystem::PlatformSupportsInputType(ECommonInputType InInputType) const
 {
-	bool bPlatformSupportsInput = FCommonInputBase::GetCurrentBasePlatformData().SupportsInputType(InInputType);
+	bool bPlatformSupportsInput = UCommonInputPlatformSettings::Get()->SupportsInputType(InInputType);
 	switch (InInputType)
 	{
 	case ECommonInputType::MouseAndKeyboard:
diff --git a/Engine/Plugins/Experimental/CommonUI/Source/CommonInput/Public/CommonInputBaseTypes.h b/Engine/Plugins/Experimental/CommonUI/Source/CommonInput/Public/CommonInputBaseTypes.h
--- a/Engine/Plugins/Experimental/CommonUI/Source/CommonInput/Public/CommonInputBaseTypes.h
+++ b/Engine/Plugins/Experimental/CommonUI/Source/CommonInput/Public/CommonInputBaseTypes.h
@@ -13,6 +13,7 @@
 #include "Styling/SlateBrush.h"
 #include "Engine/DataTable.h"
 #include "UObject/SoftObjectPtr.h"
+#include "Engine/PlatformSettings.h"
 #include "CommonInputBaseTypes.generated.h"
 
 
@@ -103,6 +104,12 @@ public:
 
 	virtual void PreSave(FObjectPreSaveContext ObjectSaveContext) override;
 
+private:
+#if WITH_EDITORONLY_DATA
+	UPROPERTY(Transient, EditAnywhere, Category = "Editor")
+	int32 SetButtonImageHeightTo = 0;
+#endif
+
 public:
 	UPROPERTY(EditDefaultsOnly, Category = "Properties")
 	ECommonInputType InputType;
@@ -124,8 +131,90 @@ public:
 
 	UFUNCTION()
 	static const TArray<FName>& GetRegisteredGamepads();
+
+private:
+#if WITH_EDITOR
+	virtual void PostEditChangeProperty(struct FPropertyChangedEvent& PropertyChangedEvent) override;
+#endif
 };
 
+UCLASS(config = Game, defaultconfig)
+class COMMONINPUT_API UCommonInputPlatformSettings : public UPlatformSettings
+{
+	GENERATED_BODY()
+
+	friend class UCommonInputSettings;
+
+public:
+	UCommonInputPlatformSettings();
+
+	virtual void PostInitProperties() override;
+
+	static UCommonInputPlatformSettings* Get()
+	{
+		return UPlatformSettings::GetSettingsForPlatform<UCommonInputPlatformSettings>();
+	}
+
+	bool TryGetInputBrush(FSlateBrush& OutBrush, FKey Key, ECommonInputType InputType, const FName& GamepadName) const;
+	bool TryGetInputBrush(FSlateBrush& OutBrush, const TArray<FKey>& Keys, ECommonInputType InputType,  const FName& GamepadName) const;
+
+	ECommonInputType GetDefaultInputType() const
+	{
+		return DefaultInputType;
+	}
+
+	bool SupportsInputType(ECommonInputType InputType) const;
+
+	const FName GetDefaultGamepadName() const
+	{
+		return DefaultGamepadName;
+	}
+
+	bool CanChangeGamepadType() const 
+	{
+		return bCanChangeGamepadType;
+	}
+
+	TArray<TSoftClassPtr<UCommonInputBaseControllerData>> GetControllerData()
+	{
+		return ControllerData;
+	}
+
+
+#if WITH_EDITOR
+	void PostEditChangeProperty(struct FPropertyChangedEvent& PropertyChangedEvent);
+#endif
+
+protected:
+	void InitializeControllerData() const;
+	virtual void InitializePlatformDefaults();
+
+	UPROPERTY(config, EditAnywhere, Category = "Default")
+	ECommonInputType DefaultInputType;
+
+	UPROPERTY(config, EditAnywhere, Category = "Default")
+	bool bSupportsMouseAndKeyboard;
+
+	UPROPERTY(config, EditAnywhere, Category = "Default")
+	bool bSupportsTouch;
+
+	UPROPERTY(config, EditAnywhere, Category = "Default")
+	bool bSupportsGamepad;
+
+	UPROPERTY(config, EditAnywhere, Category = "Default", Meta = (EditCondition = "bSupportsGamepad"))
+	FName DefaultGamepadName;
+
+	UPROPERTY(config, EditAnywhere, Category = "Default", Meta = (EditCondition = "bSupportsGamepad"))
+	bool bCanChangeGamepadType;
+
+	UPROPERTY(config, EditAnywhere, Category = "Default", Meta = (TitleProperty = "Key"))
+	TArray<TSoftClassPtr<UCommonInputBaseControllerData>> ControllerData;
+
+	UPROPERTY(Transient)
+	mutable TArray<TSubclassOf<UCommonInputBaseControllerData>> ControllerDataClasses;
+};
+
+/* DEPRECATED Legacy! */
 USTRUCT()
 struct COMMONINPUT_API FCommonInputPlatformBaseData
 {
@@ -136,7 +225,6 @@ struct COMMONINPUT_API FCommonInputPlatformBaseData
 public:
 	FCommonInputPlatformBaseData()
 	{
-		bSupported = false;
 		DefaultInputType = ECommonInputType::Gamepad;
 		bSupportsMouseAndKeyboard = false;
 		bSupportsGamepad = true;
@@ -197,15 +285,12 @@ public:
 
 protected:
 	UPROPERTY(EditDefaultsOnly, Category = "Properties")
-	bool bSupported;
-
-	UPROPERTY(EditDefaultsOnly, Category = "Properties", Meta = (EditCondition = "bSupported"))
 	ECommonInputType DefaultInputType;
 
-	UPROPERTY(EditDefaultsOnly, Category = "Properties", Meta = (EditCondition = "bSupported"))
+	UPROPERTY(EditDefaultsOnly, Category = "Properties")
 	bool bSupportsMouseAndKeyboard;
 
-	UPROPERTY(EditDefaultsOnly, Category = "Gamepad", Meta = (EditCondition = "bSupported"))
+	UPROPERTY(EditDefaultsOnly, Category = "Gamepad")
 	bool bSupportsGamepad;
 
 	UPROPERTY(EditDefaultsOnly, Category = "Gamepad", Meta = (EditCondition = "bSupportsGamepad"))
@@ -214,10 +299,10 @@ protected:
 	UPROPERTY(EditDefaultsOnly, Category = "Gamepad", Meta = (EditCondition = "bSupportsGamepad"))
 	bool bCanChangeGamepadType;
 
-	UPROPERTY(EditDefaultsOnly, Category = "Properties", Meta = (EditCondition = "bSupported"))
+	UPROPERTY(EditDefaultsOnly, Category = "Properties")
 	bool bSupportsTouch;
 
-	UPROPERTY(EditDefaultsOnly, Category = "Properties", Meta = (TitleProperty = "Key", EditCondition = "bSupported"))
+	UPROPERTY(EditDefaultsOnly, Category = "Properties", Meta = (TitleProperty = "Key"))
 	TArray<TSoftClassPtr<UCommonInputBaseControllerData>> ControllerData;
 
 	UPROPERTY(Transient)
@@ -231,7 +316,5 @@ public:
 
 	COMMONINPUT_API static UCommonInputSettings* GetInputSettings();
 
-	COMMONINPUT_API static void GetCurrentPlatformDefaults(ECommonInputType& OutDefaultInputType, FName& OutDefaultGamepadName);
-
-	COMMONINPUT_API static FCommonInputPlatformBaseData GetCurrentBasePlatformData();
+	COMMONINPUT_API void GetCurrentPlatformDefaults(ECommonInputType& OutDefaultInputType, FName& OutDefaultGamepadName);
 };
diff --git a/Engine/Plugins/Experimental/CommonUI/Source/CommonInput/Public/CommonInputSettings.h b/Engine/Plugins/Experimental/CommonUI/Source/CommonInput/Public/CommonInputSettings.h
--- a/Engine/Plugins/Experimental/CommonUI/Source/CommonInput/Public/CommonInputSettings.h
+++ b/Engine/Plugins/Experimental/CommonUI/Source/CommonInput/Public/CommonInputSettings.h
@@ -13,10 +13,11 @@
 #include "CommonInputBaseTypes.h"
 
 #include "Engine/DataTable.h"
+#include "Engine/PlatformSettings.h"
 #include "CommonInputSettings.generated.h"
 
 UCLASS(config = Game, defaultconfig)
-class COMMONINPUT_API UCommonInputSettings : public UObject
+class COMMONINPUT_API UCommonInputSettings : public UDeveloperSettings
 {
 	GENERATED_BODY()
 
@@ -36,9 +37,6 @@ public:
 	FDataTableRowHandle GetDefaultClickAction() const;
 	FDataTableRowHandle GetDefaultBackAction() const;
 
-	void GetCurrentPlatformDefaults(ECommonInputType& OutDefaultInputType, FName& OutDefaultGamepadName) const;
-	FCommonInputPlatformBaseData GetCurrentPlatform() const;
-
 	bool GetEnableInputMethodThrashingProtection() const { return bEnableInputMethodThrashingProtection; }
 
 	int32 GetInputMethodThrashingLimit() const { return InputMethodThrashingLimit; }
@@ -50,12 +48,17 @@ public:
 	bool GetAllowOutOfFocusDeviceInput() const { return bAllowOutOfFocusDeviceInput; }
 
 private:
+	virtual void PostInitProperties() override;
+
 	/** Create a derived asset from UCommonUIInputData to store Input data for your game.*/
-	UPROPERTY(config, EditAnywhere, Category = "Input")
+	UPROPERTY(config, EditAnywhere, Category = "Input", Meta=(AllowAbstract=false))
 	TSoftClassPtr<UCommonUIInputData> InputData;
+	
+	UPROPERTY(EditAnywhere, Category = "Input")
+	FPerPlatformSettings PlatformInput;
 
-	UPROPERTY(config, EditAnywhere, EditFixedSize, Category = "Input", Meta = (GetOptions = GetRegisteredPlatforms))
-	TMap<FName, FCommonInputPlatformBaseData> CommonInputPlatformData;
+	UPROPERTY(config)
+	TMap<FName, FCommonInputPlatformBaseData> CommonInputPlatformData_DEPRECATED;
 
 	UPROPERTY(config, EditAnywhere, Category = "Thrashing Settings")
 	bool bEnableInputMethodThrashingProtection = true;
@@ -72,9 +75,6 @@ private:
 	UPROPERTY(config, EditAnywhere, Category = "Input")
 	bool bAllowOutOfFocusDeviceInput = false;
 
-	UFUNCTION()
-	static const TArray<FName>& GetRegisteredPlatforms();
-
 private:
 	void LoadInputData();
 
@@ -82,7 +82,4 @@ private:
 
 	UPROPERTY(Transient)
 	TSubclassOf<UCommonUIInputData> InputDataClass;
-
-	UPROPERTY(Transient)
-	FCommonInputPlatformBaseData CurrentPlatform;
 };
diff --git a/Engine/Plugins/Experimental/CommonUI/Source/CommonUI/Private/CommonUISubsystemBase.cpp b/Engine/Plugins/Experimental/CommonUI/Source/CommonUI/Private/CommonUISubsystemBase.cpp
--- a/Engine/Plugins/Experimental/CommonUI/Source/CommonUI/Private/CommonUISubsystemBase.cpp
+++ b/Engine/Plugins/Experimental/CommonUI/Source/CommonUI/Private/CommonUISubsystemBase.cpp
@@ -51,7 +51,7 @@ FSlateBrush UCommonUISubsystemBase::GetInputActionButtonIcon(const FDataTableRow
 			}
 
 			FSlateBrush SlateBrush;
-			if (FCommonInputBase::GetCurrentBasePlatformData().TryGetInputBrush(SlateBrush, InputTypeInfo.GetKey(), InputType, GamepadName))
+			if (UCommonInputPlatformSettings::Get()->TryGetInputBrush(SlateBrush, InputTypeInfo.GetKey(), InputType, GamepadName))
 			{
 				return SlateBrush;
 			}
diff --git a/Engine/Plugins/Experimental/CommonUI/Source/CommonUI/Private/CommonUITypes.cpp b/Engine/Plugins/Experimental/CommonUI/Source/CommonUI/Private/CommonUITypes.cpp
--- a/Engine/Plugins/Experimental/CommonUI/Source/CommonUI/Private/CommonUITypes.cpp
+++ b/Engine/Plugins/Experimental/CommonUI/Source/CommonUI/Private/CommonUITypes.cpp
@@ -138,7 +138,7 @@ FSlateBrush FCommonInputActionDataBase::GetCurrentInputActionIcon(const UCommonI
 	//}
 
 	FSlateBrush SlateBrush;
-	if (FCommonInputBase::GetCurrentBasePlatformData().TryGetInputBrush(SlateBrush, CurrentInputTypeInfo.GetKey(), CommonInputSubsystem->GetCurrentInputType(), CommonInputSubsystem->GetCurrentGamepadName()))
+	if (UCommonInputPlatformSettings::Get()->TryGetInputBrush(SlateBrush, CurrentInputTypeInfo.GetKey(), CommonInputSubsystem->GetCurrentInputType(), CommonInputSubsystem->GetCurrentGamepadName()))
 	{
 		return SlateBrush;
 	}
@@ -225,7 +225,7 @@ FSlateBrush CommonUI::GetIconForInputActions(const UCommonInputSubsystem* Common
 	}
 
 	FSlateBrush SlateBrush;
-	if (FCommonInputBase::GetCurrentBasePlatformData().TryGetInputBrush(SlateBrush, Keys, CommonInputSubsystem->GetCurrentInputType(), CommonInputSubsystem->GetCurrentGamepadName()))
+	if (UCommonInputPlatformSettings::Get()->TryGetInputBrush(SlateBrush, Keys, CommonInputSubsystem->GetCurrentInputType(), CommonInputSubsystem->GetCurrentGamepadName()))
 	{
 		return SlateBrush;
 	}
diff --git a/Engine/Source/Runtime/Core/Private/Misc/DataDrivenPlatformInfoRegistry.cpp b/Engine/Source/Runtime/Core/Private/Misc/DataDrivenPlatformInfoRegistry.cpp
--- a/Engine/Source/Runtime/Core/Private/Misc/DataDrivenPlatformInfoRegistry.cpp
+++ b/Engine/Source/Runtime/Core/Private/Misc/DataDrivenPlatformInfoRegistry.cpp
@@ -287,7 +287,6 @@ static void LoadDDPIIniSettings(const FConfigFile& IniFile, FDataDrivenPlatformI
 	checkf(Info.GlobalIdentifier != FGuid(), TEXT("Platform %s didn't have a valid GlobalIdentifier set in DataDrivenPlatformInfo.ini"), *PlatformName.ToString());
 
 	// NOTE: add more settings here!
-	DDPIGetBool(IniFile, TEXT("bIsInteractablePlatform"), Info.bIsInteractablePlatform);
 	DDPIGetBool(IniFile, TEXT("bHasDedicatedGamepad"), Info.bHasDedicatedGamepad);
 	DDPIGetBool(IniFile, TEXT("bDefaultInputStandardKeyboard"), Info.bDefaultInputStandardKeyboard);
 
diff --git a/Engine/Source/Runtime/Core/Public/Misc/DataDrivenPlatformInfoRegistry.h b/Engine/Source/Runtime/Core/Public/Misc/DataDrivenPlatformInfoRegistry.h
--- a/Engine/Source/Runtime/Core/Public/Misc/DataDrivenPlatformInfoRegistry.h
+++ b/Engine/Source/Runtime/Core/Public/Misc/DataDrivenPlatformInfoRegistry.h
@@ -128,9 +128,6 @@ struct FDataDrivenPlatformInfo
 	bool Freezing_bAlignBases = false;
 	bool Freezing_bWithRayTracing = false;
 
-	// True if users will actually interact with this plaform, IE: not a GDK
-	bool bIsInteractablePlatform = false;
-
 	// True if this platform has a non-generic gamepad specifically associated with it
 	bool bHasDedicatedGamepad = false;
 
