CommonUI - No longer forcing CommonUI to ClientOnly, this is problematic as it forces game modules (even small ones) to break up thier UI and "game code" into different modules, which is too large of a burden to enforce across every game at this level, so CommonUI no longer enforces it, if however a game chooseses to split things up, that's their choice.

diff --git a/Engine/Plugins/Experimental/CommonUI/CommonUI.uplugin b/Engine/Plugins/Experimental/CommonUI/CommonUI.uplugin
--- a/Engine/Plugins/Experimental/CommonUI/CommonUI.uplugin
+++ b/Engine/Plugins/Experimental/CommonUI/CommonUI.uplugin
@@ -18,7 +18,7 @@
 	[
 		{
 			"Name" : "CommonUI",
-			"Type" : "ClientOnly",
+			"Type" : "Runtime",
 			"LoadingPhase" : "Default"
 		},
 		{
@@ -28,7 +28,7 @@
 		},
 		{
 			"Name" : "CommonInput",
-			"Type" : "ClientOnly",
+			"Type" : "Runtime",
 			"LoadingPhase" : "Default"		
 		}
 	],
diff --git a/Engine/Plugins/Experimental/CommonUI/Source/CommonUI/Private/CommonUIModule.cpp b/Engine/Plugins/Experimental/CommonUI/Source/CommonUI/Private/CommonUIModule.cpp
--- a/Engine/Plugins/Experimental/CommonUI/Source/CommonUI/Private/CommonUIModule.cpp
+++ b/Engine/Plugins/Experimental/CommonUI/Source/CommonUI/Private/CommonUIModule.cpp
@@ -87,7 +87,7 @@ void FCommonUIModule::CreateUISettings()
 			CommonUIEditorSettings);
 
 		SettingsModule->RegisterSettings(TEXT("Project"), TEXT("Plugins"), TEXT("CommonUIInputSettings"),
-			NSLOCTEXT("CommonUIPlugin", "CommonUIInputSettingsName", "UI Input Settings"),
+			NSLOCTEXT("CommonUIPlugin", "CommonUIInputSettingsName", "Common UI Input Settings"),
 			NSLOCTEXT("CommonUIPlugin", "CommonUIInputSettingsDescription", "Establish project-wide UI input settings and action mappings."),
 			GetMutableDefault<UCommonUIInputSettings>());
 	}
