CIS - Fixing some legacy common UI.

diff --git a/Engine/Plugins/Experimental/CommonUI/Source/CommonInput/Public/CommonInputBaseTypes.h b/Engine/Plugins/Experimental/CommonUI/Source/CommonInput/Public/CommonInputBaseTypes.h
--- a/Engine/Plugins/Experimental/CommonUI/Source/CommonInput/Public/CommonInputBaseTypes.h
+++ b/Engine/Plugins/Experimental/CommonUI/Source/CommonInput/Public/CommonInputBaseTypes.h
@@ -316,5 +316,5 @@ public:
 
 	COMMONINPUT_API static UCommonInputSettings* GetInputSettings();
 
-	COMMONINPUT_API void GetCurrentPlatformDefaults(ECommonInputType& OutDefaultInputType, FName& OutDefaultGamepadName);
+	COMMONINPUT_API static void GetCurrentPlatformDefaults(ECommonInputType& OutDefaultInputType, FName& OutDefaultGamepadName);
 };
