CommonInput - Don't try to load alll platforms input data outside of the editor.

diff --git a/Engine/Plugins/Experimental/CommonUI/Source/CommonInput/Private/CommonInputSettings.cpp b/Engine/Plugins/Experimental/CommonUI/Source/CommonInput/Private/CommonInputSettings.cpp
--- a/Engine/Plugins/Experimental/CommonUI/Source/CommonInput/Private/CommonInputSettings.cpp
+++ b/Engine/Plugins/Experimental/CommonUI/Source/CommonInput/Private/CommonInputSettings.cpp
@@ -13,7 +13,9 @@ UCommonInputSettings::UCommonInputSettings(const FObjectInitializer& Initializer
 	: Super(Initializer)
 	, bInputDataLoaded(false)
 {
+#if WITH_EDITOR
 	PlatformInput.Settings = UPlatformSettings::GetAllPlatformSettings<UCommonInputPlatformSettings>();
+#endif
 }
 
 void UCommonInputSettings::LoadData()
