Common UI - Make sure to link DeveloperSettings.

diff --git a/Engine/Plugins/Experimental/CommonUI/Source/CommonUI/CommonUI.Build.cs b/Engine/Plugins/Experimental/CommonUI/Source/CommonUI/CommonUI.Build.cs
--- a/Engine/Plugins/Experimental/CommonUI/Source/CommonUI/CommonUI.Build.cs
+++ b/Engine/Plugins/Experimental/CommonUI/Source/CommonUI/CommonUI.Build.cs
@@ -19,7 +19,7 @@ public class CommonUI : ModuleRules
                 "CommonInput",
 				"MediaAssets",
 				"GameplayTags"
-            }
+			}
 		);
 
 		PrivateDependencyModuleNames.AddRange(
@@ -29,8 +29,9 @@ public class CommonUI : ModuleRules
                 "Analytics",
 				"AnalyticsET",
                 "EngineSettings",
-				"AudioMixer"
-            }
+				"AudioMixer",
+				"DeveloperSettings"
+			}
 		);
 
         PrivateIncludePaths.AddRange(
