Make sure we don't crash in CommonInput if the platforms are missing.

diff --git a/Engine/Plugins/Experimental/CommonUI/Source/CommonInput/Private/CommonInputSettings.cpp b/Engine/Plugins/Experimental/CommonUI/Source/CommonInput/Private/CommonInputSettings.cpp
--- a/Engine/Plugins/Experimental/CommonUI/Source/CommonInput/Private/CommonInputSettings.cpp
+++ b/Engine/Plugins/Experimental/CommonUI/Source/CommonInput/Private/CommonInputSettings.cpp
@@ -132,6 +132,8 @@ void UCommonInputSettings::PostInitProperties()
 				Settings->DefaultGamepadName = OriginalData.DefaultGamepadName;
 				Settings->DefaultInputType = OriginalData.DefaultInputType;
 				Settings->ControllerData = OriginalData.ControllerData;
+
+				Settings->UpdateDefaultConfigFile();
 			}
 			else if (PlatformData.Key == FCommonInputDefaults::PlatformPC)
 			{
@@ -142,13 +144,18 @@ void UCommonInputSettings::PostInitProperties()
 
 				for (UCommonInputPlatformSettings* PCPlatform : PCPlatforms)
 				{
-					PCPlatform->bSupportsMouseAndKeyboard = OriginalData.bSupportsMouseAndKeyboard;
-					PCPlatform->bSupportsGamepad = OriginalData.bSupportsGamepad;
-					PCPlatform->bSupportsTouch = OriginalData.bSupportsTouch;
-					PCPlatform->bCanChangeGamepadType = OriginalData.bCanChangeGamepadType;
-					PCPlatform->DefaultGamepadName = OriginalData.DefaultGamepadName;
-					PCPlatform->DefaultInputType = OriginalData.DefaultInputType;
-					PCPlatform->ControllerData = OriginalData.ControllerData;
+					if (PCPlatform)
+					{
+						PCPlatform->bSupportsMouseAndKeyboard = OriginalData.bSupportsMouseAndKeyboard;
+						PCPlatform->bSupportsGamepad = OriginalData.bSupportsGamepad;
+						PCPlatform->bSupportsTouch = OriginalData.bSupportsTouch;
+						PCPlatform->bCanChangeGamepadType = OriginalData.bCanChangeGamepadType;
+						PCPlatform->DefaultGamepadName = OriginalData.DefaultGamepadName;
+						PCPlatform->DefaultInputType = OriginalData.DefaultInputType;
+						PCPlatform->ControllerData = OriginalData.ControllerData;
+
+						PCPlatform->UpdateDefaultConfigFile();
+					}
 				}
 			}
 		}
