CommonUI - Don't include the UITag.h header in CommonActivatableWidget.h, that dependency got removed in a refactor, and it causes a lot of other modules to suddenly depend on GameplayTags.

diff --git a/Engine/Plugins/Experimental/CommonUI/Source/CommonUI/Public/CommonActivatableWidget.h b/Engine/Plugins/Experimental/CommonUI/Source/CommonUI/Public/CommonActivatableWidget.h
--- a/Engine/Plugins/Experimental/CommonUI/Source/CommonUI/Public/CommonActivatableWidget.h
+++ b/Engine/Plugins/Experimental/CommonUI/Source/CommonUI/Public/CommonActivatableWidget.h
@@ -4,7 +4,6 @@
 
 #include "CommonUserWidget.h"
 #include "Input/UIActionBindingHandle.h"
-#include "UITag.h"
 #include "CommonActivatableWidget.generated.h"
 
 DECLARE_DYNAMIC_MULTICAST_DELEGATE(FOnWidgetActivationChanged);
