Fixing not enabling LastInputType, after the changes to the per-platform settings.

diff --git a/Engine/Plugins/Experimental/CommonUI/Source/CommonInput/Private/CommonInputSubsystem.cpp b/Engine/Plugins/Experimental/CommonUI/Source/CommonInput/Private/CommonInputSubsystem.cpp
--- a/Engine/Plugins/Experimental/CommonUI/Source/CommonInput/Private/CommonInputSubsystem.cpp
+++ b/Engine/Plugins/Experimental/CommonUI/Source/CommonInput/Private/CommonInputSubsystem.cpp
@@ -285,7 +285,7 @@ void UCommonInputSubsystem::Initialize(FSubsystemCollectionBase& Collection)
 
 	FCommonInputBase::GetInputSettings()->LoadData();
 	GamepadInputType = UPlatformSettings::GetSettingsForPlatform<UCommonInputPlatformSettings>()->GetDefaultGamepadName();
-	CurrentInputType = UPlatformSettings::GetSettingsForPlatform<UCommonInputPlatformSettings>()->GetDefaultInputType();
+	CurrentInputType = LastInputType = UPlatformSettings::GetSettingsForPlatform<UCommonInputPlatformSettings>()->GetDefaultInputType();
 
 	CommonInputPreprocessor = MakeShared<FCommonInputPreprocessor>(*this);
 	FSlateApplication::Get().RegisterInputPreProcessor(CommonInputPreprocessor, 0);
