Fix CommonUI platform settings attempting to load assets too early.

diff --git a/Engine/Plugins/Experimental/CommonUI/Source/CommonInput/Private/CommonInputBaseTypes.cpp b/Engine/Plugins/Experimental/CommonUI/Source/CommonInput/Private/CommonInputBaseTypes.cpp
--- a/Engine/Plugins/Experimental/CommonUI/Source/CommonInput/Private/CommonInputBaseTypes.cpp
+++ b/Engine/Plugins/Experimental/CommonUI/Source/CommonInput/Private/CommonInputBaseTypes.cpp
@@ -204,9 +204,9 @@ UCommonInputPlatformSettings::UCommonInputPlatformSettings()
 	DefaultGamepadName = FCommonInputDefaults::GamepadGeneric;
 }
 
-void UCommonInputPlatformSettings::PostInitProperties()
+void UCommonInputPlatformSettings::PostLoad()
 {
-	Super::PostInitProperties();
+	Super::PostLoad();
 
 	ControllerDataClasses.Reset();
 	InitializeControllerData();
diff --git a/Engine/Plugins/Experimental/CommonUI/Source/CommonInput/Public/CommonInputBaseTypes.h b/Engine/Plugins/Experimental/CommonUI/Source/CommonInput/Public/CommonInputBaseTypes.h
--- a/Engine/Plugins/Experimental/CommonUI/Source/CommonInput/Public/CommonInputBaseTypes.h
+++ b/Engine/Plugins/Experimental/CommonUI/Source/CommonInput/Public/CommonInputBaseTypes.h
@@ -148,7 +148,7 @@ class COMMONINPUT_API UCommonInputPlatformSettings : public UPlatformSettings
 public:
 	UCommonInputPlatformSettings();
 
-	virtual void PostInitProperties() override;
+	virtual void PostLoad() override;
 
 	static UCommonInputPlatformSettings* Get()
 	{
