CommonUI - Adding the CommonHardwareVisibilityBorder, it will replace the CommonVisibilityWidgetBase.  Allows doing more sophisticated tag based logic for different platforms based on features as opposed to platform names, which don't work particularly well as you add new platforms, so instead encoding visibility rules based on hardware features.  Features definable by game teams.

diff --git a/Engine/Plugins/Experimental/CommonUI/Source/CommonInput/Private/CommonInputBaseTypes.cpp b/Engine/Plugins/Experimental/CommonUI/Source/CommonInput/Private/CommonInputBaseTypes.cpp
--- a/Engine/Plugins/Experimental/CommonUI/Source/CommonInput/Private/CommonInputBaseTypes.cpp
+++ b/Engine/Plugins/Experimental/CommonUI/Source/CommonInput/Private/CommonInputBaseTypes.cpp
@@ -119,6 +119,14 @@ void UCommonInputBaseControllerData::PreSave(FObjectPreSaveContext ObjectSaveCon
 	}
 }
 
+void UCommonInputBaseControllerData::PostLoad()
+{
+	Super::PostLoad();
+
+	// Have to clear it even though it's transient because it's saved into the CDO.
+	SetButtonImageHeightTo = 0;
+}
+
 const TArray<FName>& UCommonInputBaseControllerData::GetRegisteredGamepads()
 {
 	auto GenerateRegisteredGamepads = []()
@@ -171,19 +179,25 @@ void UCommonInputBaseControllerData::PostEditChangeProperty(struct FPropertyChan
 				for (FCommonInputKeyBrushConfiguration& BrushConfig : InputBrushDataMap)
 				{
 					FVector2D NewBrushSize = BrushConfig.KeyBrush.GetImageSize();
-					NewBrushSize.X = FMath::RoundToInt(NewBrushSize.X * (SetButtonImageHeightTo / NewBrushSize.X));
-					NewBrushSize.Y = SetButtonImageHeightTo;
+					if (NewBrushSize.X != 0 && NewBrushSize.Y != 0)
+					{
+						NewBrushSize.X = FMath::RoundToInt(SetButtonImageHeightTo * (NewBrushSize.X / NewBrushSize.Y));
+						NewBrushSize.Y = SetButtonImageHeightTo;
 
-					BrushConfig.KeyBrush.SetImageSize(NewBrushSize);
+						BrushConfig.KeyBrush.SetImageSize(NewBrushSize);
+					}
 				}
 
 				for (FCommonInputKeySetBrushConfiguration& BrushConfig : InputBrushKeySets)
 				{
 					FVector2D NewBrushSize = BrushConfig.KeyBrush.GetImageSize();
-					NewBrushSize.X = FMath::RoundToInt(NewBrushSize.X * (SetButtonImageHeightTo / NewBrushSize.X));
-					NewBrushSize.Y = SetButtonImageHeightTo;
+					if (NewBrushSize.X != 0 && NewBrushSize.Y != 0)
+					{
+						NewBrushSize.X = FMath::RoundToInt(SetButtonImageHeightTo * (NewBrushSize.X / NewBrushSize.Y));
+						NewBrushSize.Y = SetButtonImageHeightTo;
 
-					BrushConfig.KeyBrush.SetImageSize(NewBrushSize);
+						BrushConfig.KeyBrush.SetImageSize(NewBrushSize);
+					}
 				}
 			}
 		}
diff --git a/Engine/Plugins/Experimental/CommonUI/Source/CommonInput/Public/CommonInputBaseTypes.h b/Engine/Plugins/Experimental/CommonUI/Source/CommonInput/Public/CommonInputBaseTypes.h
--- a/Engine/Plugins/Experimental/CommonUI/Source/CommonInput/Public/CommonInputBaseTypes.h
+++ b/Engine/Plugins/Experimental/CommonUI/Source/CommonInput/Public/CommonInputBaseTypes.h
@@ -103,6 +103,7 @@ public:
 	virtual bool TryGetInputBrush(FSlateBrush& OutBrush, const TArray<FKey>& Keys) const;
 
 	virtual void PreSave(FObjectPreSaveContext ObjectSaveContext) override;
+	virtual void PostLoad() override;
 
 private:
 #if WITH_EDITORONLY_DATA
@@ -111,22 +112,25 @@ private:
 #endif
 
 public:
-	UPROPERTY(EditDefaultsOnly, Category = "Properties")
+	UPROPERTY(EditDefaultsOnly, Category = "Default")
 	ECommonInputType InputType;
 	
-	UPROPERTY(EditDefaultsOnly, Category = "Properties", Meta = (GetOptions = GetRegisteredGamepads))
+	UPROPERTY(EditDefaultsOnly, Category = "Gamepad", meta=(EditCondition="InputType == ECommonInputType::Gamepad"))
 	FName GamepadName;
 
-	UPROPERTY(EditDefaultsOnly, Category = "Properties")
+	UPROPERTY(EditDefaultsOnly, Category = "Gamepad", meta = (EditCondition = "InputType == ECommonInputType::Gamepad"))
+	FText GamepadDisplayName;
+
+	UPROPERTY(EditDefaultsOnly, Category = "Display")
 	TSoftObjectPtr<UTexture2D> ControllerTexture;
 
-	UPROPERTY(EditDefaultsOnly, Category = "Properties")
+	UPROPERTY(EditDefaultsOnly, Category = "Display")
 	TSoftObjectPtr<UTexture2D> ControllerButtonMaskTexture;
 
-	UPROPERTY(EditDefaultsOnly, Category = "Properties", Meta = (TitleProperty = "Key"))
+	UPROPERTY(EditDefaultsOnly, Category = "Display", Meta = (TitleProperty = "Key"))
 	TArray<FCommonInputKeyBrushConfiguration> InputBrushDataMap;
 
-	UPROPERTY(EditDefaultsOnly, Category = "Properties", Meta = (TitleProperty = "Keys"))
+	UPROPERTY(EditDefaultsOnly, Category = "Display", Meta = (TitleProperty = "Keys"))
 	TArray<FCommonInputKeySetBrushConfiguration> InputBrushKeySets;
 
 	UFUNCTION()
@@ -207,7 +211,7 @@ protected:
 	UPROPERTY(config, EditAnywhere, Category = "Default", Meta = (EditCondition = "bSupportsGamepad"))
 	bool bCanChangeGamepadType;
 
-	UPROPERTY(config, EditAnywhere, Category = "Default", Meta = (TitleProperty = "Key"))
+	UPROPERTY(config, EditAnywhere, Category = "Default", Meta = (TitleProperty = "InputType"))
 	TArray<TSoftClassPtr<UCommonInputBaseControllerData>> ControllerData;
 
 	UPROPERTY(Transient)
@@ -302,7 +306,7 @@ protected:
 	UPROPERTY(EditDefaultsOnly, Category = "Properties")
 	bool bSupportsTouch;
 
-	UPROPERTY(EditDefaultsOnly, Category = "Properties", Meta = (TitleProperty = "Key"))
+	UPROPERTY(EditDefaultsOnly, Category = "Properties", Meta = (TitleProperty = "GamepadName"))
 	TArray<TSoftClassPtr<UCommonInputBaseControllerData>> ControllerData;
 
 	UPROPERTY(Transient)
diff --git a/Engine/Plugins/Experimental/CommonUI/Source/CommonUI/Private/CommonHardwareVisibilityBorder.cpp b/Engine/Plugins/Experimental/CommonUI/Source/CommonUI/Private/CommonHardwareVisibilityBorder.cpp
new file mode 100644
--- /dev/null
+++ b/Engine/Plugins/Experimental/CommonUI/Source/CommonUI/Private/CommonHardwareVisibilityBorder.cpp
@@ -0,0 +1,62 @@
+// Copyright Epic Games, Inc. All Rights Reserved.
+
+#include "CommonHardwareVisibilityBorder.h"
+#include "CommonUIPrivatePCH.h"
+#include "CommonUISubsystemBase.h"
+
+#include "CommonInputBaseTypes.h"
+#include "Misc/DataDrivenPlatformInfoRegistry.h"
+#include "CommonUIVisibilitySubsystem.h"
+
+UCommonHardwareVisibilityBorder::UCommonHardwareVisibilityBorder(const FObjectInitializer& ObjectInitializer)
+	: Super(ObjectInitializer)
+	, VisibleType(ESlateVisibility::SelfHitTestInvisible)
+	, HiddenType(ESlateVisibility::Collapsed)
+{
+	Padding = FMargin(0,0,0,0);
+}
+
+void UCommonHardwareVisibilityBorder::OnWidgetRebuilt()
+{
+	Super::OnWidgetRebuilt();
+	UpdateVisibility();
+	ListenToInputMethodChanged();
+}
+
+void UCommonHardwareVisibilityBorder::UpdateVisibility(UCommonUIVisibilitySubsystem* VisSystem)
+{
+	if (!IsDesignTime())
+	{
+		const UCommonUIVisibilitySubsystem* CommonInputSubsystem = VisSystem ? VisSystem : UCommonUIVisibilitySubsystem::Get(GetOwningLocalPlayer());
+		if (ensure(CommonInputSubsystem))
+		{
+			bool bVisibleForFeatures = true;
+
+			if (!VisibilityQuery.IsEmpty())
+			{
+				bVisibleForFeatures = VisibilityQuery.Matches(CommonInputSubsystem->GetVisibilityTags());
+			}
+
+			SetVisibility(bVisibleForFeatures ? VisibleType : HiddenType);
+		}
+	}
+}
+
+void UCommonHardwareVisibilityBorder::ListenToInputMethodChanged()
+{
+	if (IsDesignTime())
+	{
+		return;
+	}
+
+	if (UCommonUIVisibilitySubsystem* CommonInputSubsystem = UCommonUIVisibilitySubsystem::Get(GetOwningLocalPlayer()))
+	{
+		CommonInputSubsystem->OnVisibilityTagsChanged.RemoveAll(this);
+		CommonInputSubsystem->OnVisibilityTagsChanged.AddUObject(this, &ThisClass::HandleInputMethodChanged);
+	}
+}
+
+void UCommonHardwareVisibilityBorder::HandleInputMethodChanged(UCommonUIVisibilitySubsystem* VisSystem)
+{
+	UpdateVisibility(VisSystem);
+}
diff --git a/Engine/Plugins/Experimental/CommonUI/Source/CommonUI/Private/CommonUIVisibilitySubsystem.cpp b/Engine/Plugins/Experimental/CommonUI/Source/CommonUI/Private/CommonUIVisibilitySubsystem.cpp
new file mode 100644
--- /dev/null
+++ b/Engine/Plugins/Experimental/CommonUI/Source/CommonUI/Private/CommonUIVisibilitySubsystem.cpp
@@ -0,0 +1,107 @@
+// Copyright Epic Games, Inc. All Rights Reserved.
+
+#include "CommonUIVisibilitySubsystem.h"
+
+#include "NativeGameplayTags.h"
+#include "Engine/PlatformSettings.h"
+#include "CommonInputBaseTypes.h"
+#include "Engine/LocalPlayer.h"
+#include "Engine/GameInstance.h"
+#include "CommonInputSubsystem.h"
+#include "UObject/UObjectHash.h"
+#include "ICommonUIModule.h"
+#include "CommonUISettings.h"
+
+UE_DEFINE_GAMEPLAY_TAG_STATIC(TAG_INPUT_MOUSEANDKEYBOARD, "Input.MouseAndKeyboard");
+UE_DEFINE_GAMEPLAY_TAG_STATIC(TAG_INPUT_GAMEPAD, "Input.Gamepad");
+UE_DEFINE_GAMEPLAY_TAG_STATIC(TAG_INPUT_TOUCH, "Input.Touch");
+
+UCommonUIVisibilitySubsystem* UCommonUIVisibilitySubsystem::Get(const ULocalPlayer* LocalPlayer)
+{
+	return LocalPlayer ? LocalPlayer->GetSubsystem<UCommonUIVisibilitySubsystem>() : nullptr;
+}
+
+UCommonUIVisibilitySubsystem::UCommonUIVisibilitySubsystem()
+{
+}
+
+bool UCommonUIVisibilitySubsystem::ShouldCreateSubsystem(UObject* Outer) const
+{
+	if (!CastChecked<ULocalPlayer>(Outer)->GetGameInstance()->IsDedicatedServerInstance())
+	{
+		TArray<UClass*> ChildClasses;
+		GetDerivedClasses(GetClass(), ChildClasses, false);
+
+		// Only create an instance if there is no override implementation defined elsewhere
+		return ChildClasses.Num() == 0;
+	}
+	
+	return false;
+}
+
+void UCommonUIVisibilitySubsystem::Initialize(FSubsystemCollectionBase& Collection)
+{
+	Super::Initialize(Collection);
+	Collection.InitializeDependency<UCommonInputSubsystem>();
+	
+	if (UCommonInputSubsystem* Input = GetLocalPlayer()->GetSubsystem<UCommonInputSubsystem>())
+	{
+		Input->OnInputMethodChangedNative.AddUObject(this, &ThisClass::OnInputMethodChanged);
+	}
+
+	ComputedVisibilityTags = ComputeVisibilityTags();
+}
+
+void UCommonUIVisibilitySubsystem::Deinitialize()
+{
+	Super::Deinitialize();
+}
+
+void UCommonUIVisibilitySubsystem::AddUserVisibilityCondition(const FGameplayTag& UserTag)
+{
+	UserVisibilityTags.AddTag(UserTag);
+	RefreshVisibilityTags();
+}
+
+void UCommonUIVisibilitySubsystem::RemoveUserVisibilityCondition(const FGameplayTag& UserTag)
+{
+	UserVisibilityTags.RemoveTag(UserTag);
+	RefreshVisibilityTags();
+}
+
+void UCommonUIVisibilitySubsystem::RefreshVisibilityTags()
+{
+	ComputedVisibilityTags = ComputeVisibilityTags();
+
+	OnVisibilityTagsChanged.Broadcast(this);
+}
+
+FGameplayTagContainer UCommonUIVisibilitySubsystem::ComputeVisibilityTags() const
+{
+	FGameplayTagContainer ComputedTags = UserVisibilityTags;
+
+	if (const UCommonInputSubsystem* CommonInputSubsystem = UCommonInputSubsystem::Get(GetLocalPlayer()))
+	{
+		switch(CommonInputSubsystem->GetCurrentInputType())
+		{
+			case ECommonInputType::MouseAndKeyboard:
+				ComputedTags.AddTag(TAG_INPUT_MOUSEANDKEYBOARD);
+				break;
+			case ECommonInputType::Gamepad:
+				ComputedTags.AddTag(TAG_INPUT_GAMEPAD);
+				break;
+			case ECommonInputType::Touch:
+				ComputedTags.AddTag(TAG_INPUT_TOUCH);
+				break;
+		}
+	}
+
+	ComputedTags.AppendTags(ICommonUIModule::GetSettings().GetPlatformHardwareFeatures());
+
+	return ComputedTags;
+}
+
+void UCommonUIVisibilitySubsystem::OnInputMethodChanged(ECommonInputType CurrentInputType)
+{
+	RefreshVisibilityTags();
+}
diff --git a/Engine/Plugins/Experimental/CommonUI/Source/CommonUI/Public/CommonHardwareVisibilityBorder.h b/Engine/Plugins/Experimental/CommonUI/Source/CommonUI/Public/CommonHardwareVisibilityBorder.h
new file mode 100644
--- /dev/null
+++ b/Engine/Plugins/Experimental/CommonUI/Source/CommonUI/Public/CommonHardwareVisibilityBorder.h
@@ -0,0 +1,45 @@
+// Copyright Epic Games, Inc. All Rights Reserved.
+
+#pragma once
+
+#include "CoreMinimal.h"
+#include "CommonUITypes.h"
+#include "CommonBorder.h"
+#include "GameplayTagContainer.h"
+#include "CommonHardwareVisibilityBorder.generated.h"
+
+class UCommonUIVisibilitySubsystem;
+
+/**
+ * A container that controls visibility based on Platform, Input 
+ */
+UCLASS()
+class COMMONUI_API UCommonHardwareVisibilityBorder : public UCommonBorder
+{
+	GENERATED_UCLASS_BODY()
+
+public:
+	
+
+protected:
+
+	UPROPERTY(EditAnywhere, Category = "Visibility", meta=(Categories="Input,Hardware.Feature"))
+	FGameplayTagQuery VisibilityQuery;
+
+	UPROPERTY(EditAnywhere, Category = "Visibility")
+	ESlateVisibility VisibleType;
+	
+	UPROPERTY(EditAnywhere, Category = "Visibility")
+	ESlateVisibility HiddenType;
+
+protected:
+	// Begin UWidget
+	virtual void OnWidgetRebuilt() override;
+	// End UWidget
+
+	void UpdateVisibility(UCommonUIVisibilitySubsystem* VisSystem = nullptr);
+
+	void ListenToInputMethodChanged();
+
+	void HandleInputMethodChanged(UCommonUIVisibilitySubsystem*);
+};
diff --git a/Engine/Plugins/Experimental/CommonUI/Source/CommonUI/Public/CommonUISettings.h b/Engine/Plugins/Experimental/CommonUI/Source/CommonUI/Public/CommonUISettings.h
--- a/Engine/Plugins/Experimental/CommonUI/Source/CommonUI/Public/CommonUISettings.h
+++ b/Engine/Plugins/Experimental/CommonUI/Source/CommonUI/Public/CommonUISettings.h
@@ -11,6 +11,7 @@
 #include "CommonUITypes.h"
 #include "CommonUIRichTextData.h"
 #include "CommonTextBlock.h"
+#include "GameplayTagContainer.h"
 #include "CommonUISettings.generated.h"
 
 class UMaterial;
@@ -36,6 +37,7 @@ public:
 	UCommonUIRichTextData* GetRichTextData() const;
 	const FSlateBrush& GetDefaultThrobberBrush() const;
 	UObject* GetDefaultImageResourceObject() const;
+	const FGameplayTagContainer& GetPlatformHardwareFeatures() const { return PlatformHardwareFeatures; }
 
 private:
 
@@ -57,6 +59,10 @@ private:
 	UPROPERTY(config, EditAnywhere, Category = "RichText", meta=(AllowAbstract=false))
 	TSoftClassPtr<UCommonUIRichTextData> DefaultRichTextDataClass;
 
+	/** The Default Data for rich text to show inline icon and others. */
+	UPROPERTY(config, EditAnywhere, Category = "Visibility", meta=(Categories="Hardware.Feature", ConfigHierarchyEditable))
+	FGameplayTagContainer PlatformHardwareFeatures;
+
 private:
 	void LoadEditorData();
 
diff --git a/Engine/Plugins/Experimental/CommonUI/Source/CommonUI/Public/CommonUIVisibilitySubsystem.h b/Engine/Plugins/Experimental/CommonUI/Source/CommonUI/Public/CommonUIVisibilitySubsystem.h
new file mode 100644
--- /dev/null
+++ b/Engine/Plugins/Experimental/CommonUI/Source/CommonUI/Public/CommonUIVisibilitySubsystem.h
@@ -0,0 +1,50 @@
+// Copyright Epic Games, Inc. All Rights Reserved.
+
+#pragma once
+
+#include "Subsystems/LocalPlayerSubsystem.h"
+#include "GameplayTagContainer.h"
+#include "CommonUIVisibilitySubsystem.generated.h"
+
+class UWidget;
+class ULocalPlayer;
+class APlayerController;
+struct FGameplayTagContainer;
+class UCommonUIVisibilitySubsystem;
+enum class ECommonInputType : uint8;
+
+DECLARE_DYNAMIC_MULTICAST_DELEGATE_OneParam(FHardwareVisibilityTagsChangedDynamicEvent, UCommonUIVisibilitySubsystem*, TagSubsystem);
+
+UCLASS(DisplayName = "UI Visibility Subsystem")
+class COMMONUI_API UCommonUIVisibilitySubsystem : public ULocalPlayerSubsystem
+{
+	GENERATED_BODY()
+
+public:
+	static UCommonUIVisibilitySubsystem* Get(const ULocalPlayer* LocalPlayer);
+
+	UCommonUIVisibilitySubsystem();
+	
+	virtual bool ShouldCreateSubsystem(UObject* Outer) const override;
+
+	virtual void Initialize(FSubsystemCollectionBase& Collection) override;
+	virtual void Deinitialize() override;
+
+	DECLARE_EVENT_OneParam(UCommonUIVisibilitySubsystem, FHardwareVisibilityTagsChangedEvent, UCommonUIVisibilitySubsystem*);
+	FHardwareVisibilityTagsChangedEvent OnVisibilityTagsChanged;
+
+	/** Get the hardware visibility tags currently in play.  These can change over time, if input mode changes, or other groups are removed/added. */
+	const FGameplayTagContainer& GetVisibilityTags() const { return ComputedVisibilityTags; }
+
+	void AddUserVisibilityCondition(const FGameplayTag& UserTag);
+	void RemoveUserVisibilityCondition(const FGameplayTag& UserTag);
+
+protected:
+	void RefreshVisibilityTags();
+	void OnInputMethodChanged(ECommonInputType CurrentInputType);
+	FGameplayTagContainer ComputeVisibilityTags() const;
+
+private:
+	FGameplayTagContainer ComputedVisibilityTags;
+	FGameplayTagContainer UserVisibilityTags;
+};
\ No newline at end of file
