CommonInput - Fix non-editor build.

diff --git a/Engine/Plugins/Experimental/CommonUI/Source/CommonInput/Private/CommonInputBaseTypes.cpp b/Engine/Plugins/Experimental/CommonUI/Source/CommonInput/Private/CommonInputBaseTypes.cpp
--- a/Engine/Plugins/Experimental/CommonUI/Source/CommonInput/Private/CommonInputBaseTypes.cpp
+++ b/Engine/Plugins/Experimental/CommonUI/Source/CommonInput/Private/CommonInputBaseTypes.cpp
@@ -123,8 +123,10 @@ void UCommonInputBaseControllerData::PostLoad()
 {
 	Super::PostLoad();
 
+#if WITH_EDITOR
 	// Have to clear it even though it's transient because it's saved into the CDO.
 	SetButtonImageHeightTo = 0;
+#endif
 }
 
 const TArray<FName>& UCommonInputBaseControllerData::GetRegisteredGamepads()
