CommonUI: Reworked/renamed hardware features into platform traits
- Changed various references from Hardware Feature to to Platform Traits (including the tag namespace which is now "Platform.Trait")
- Added the ability for games to force debug enables and suppressions globally (SetDebugVisibilityConditions)
- The PlatformHardwareFeatures setting is now called PlatformTraits (and is an array of tags instead of a tag container to allow for ini inhertance to work, which sadly doesn't work correctly in the current config hierarchy editor anyways; will JIRA that separately)

diff --git a/Engine/Plugins/Experimental/CommonUI/Source/CommonUI/Private/CommonUISettings.cpp b/Engine/Plugins/Experimental/CommonUI/Source/CommonUI/Private/CommonUISettings.cpp
--- a/Engine/Plugins/Experimental/CommonUI/Source/CommonUI/Private/CommonUISettings.cpp
+++ b/Engine/Plugins/Experimental/CommonUI/Source/CommonUI/Private/CommonUISettings.cpp
@@ -19,9 +19,27 @@ void UCommonUISettings::LoadData()
 #if WITH_EDITOR
 void UCommonUISettings::PostEditChangeProperty(struct FPropertyChangedEvent& PropertyChangedEvent)
 {
+	Super::PostEditChangeProperty(PropertyChangedEvent);
+
 	bDefaultDataLoaded = false;
 	AutoLoadData();
+	RebuildTraitContainer();
+}
+
+void UCommonUISettings::PostReloadConfig(FProperty* PropertyThatWasLoaded)
+{
+	Super::PostReloadConfig(PropertyThatWasLoaded);
+
+	RebuildTraitContainer();
 }
+
+void UCommonUISettings::PostInitProperties()
+{
+	Super::PostInitProperties();
+
+	RebuildTraitContainer();
+}
+
 #endif
 
 void UCommonUISettings::LoadEditorData()
@@ -59,6 +77,15 @@ void UCommonUISettings::LoadEditorData()
 	}
 }
 
+void UCommonUISettings::RebuildTraitContainer()
+{
+	PlatformTraitContainer.Reset();
+	for (FGameplayTag Trait : PlatformTraits)
+	{
+		PlatformTraitContainer.AddTag(Trait);
+	}
+}
+
 void UCommonUISettings::AutoLoadData()
 {
 	if (bAutoLoadData)
diff --git a/Engine/Plugins/Experimental/CommonUI/Source/CommonUI/Private/CommonUIVisibilitySubsystem.cpp b/Engine/Plugins/Experimental/CommonUI/Source/CommonUI/Private/CommonUIVisibilitySubsystem.cpp
--- a/Engine/Plugins/Experimental/CommonUI/Source/CommonUI/Private/CommonUIVisibilitySubsystem.cpp
+++ b/Engine/Plugins/Experimental/CommonUI/Source/CommonUI/Private/CommonUIVisibilitySubsystem.cpp
@@ -16,6 +16,11 @@ UE_DEFINE_GAMEPLAY_TAG_STATIC(TAG_INPUT_MOUSEANDKEYBOARD, "Input.MouseAndKeyboar
 UE_DEFINE_GAMEPLAY_TAG_STATIC(TAG_INPUT_GAMEPAD, "Input.Gamepad");
 UE_DEFINE_GAMEPLAY_TAG_STATIC(TAG_INPUT_TOUCH, "Input.Touch");
 
+#if WITH_EDITOR
+FGameplayTagContainer UCommonUIVisibilitySubsystem::DebugTagsToEnable;
+FGameplayTagContainer UCommonUIVisibilitySubsystem::DebugTagsToSuppress;
+#endif
+
 UCommonUIVisibilitySubsystem* UCommonUIVisibilitySubsystem::Get(const ULocalPlayer* LocalPlayer)
 {
 	return LocalPlayer ? LocalPlayer->GetSubsystem<UCommonUIVisibilitySubsystem>() : nullptr;
@@ -57,18 +62,31 @@ void UCommonUIVisibilitySubsystem::Deinitialize()
 	Super::Deinitialize();
 }
 
-void UCommonUIVisibilitySubsystem::AddUserVisibilityCondition(const FGameplayTag& UserTag)
+void UCommonUIVisibilitySubsystem::AddUserVisibilityCondition(const FGameplayTag UserTag)
 {
 	UserVisibilityTags.AddTag(UserTag);
 	RefreshVisibilityTags();
 }
 
-void UCommonUIVisibilitySubsystem::RemoveUserVisibilityCondition(const FGameplayTag& UserTag)
+void UCommonUIVisibilitySubsystem::RemoveUserVisibilityCondition(const FGameplayTag UserTag)
 {
 	UserVisibilityTags.RemoveTag(UserTag);
 	RefreshVisibilityTags();
 }
 
+#if WITH_EDITOR
+void UCommonUIVisibilitySubsystem::SetDebugVisibilityConditions(const FGameplayTagContainer& TagsToEnable, const FGameplayTagContainer& TagsToSuppress)
+{
+	DebugTagsToEnable = TagsToEnable;
+	DebugTagsToSuppress = TagsToSuppress;
+
+	for (TObjectIterator<UCommonUIVisibilitySubsystem> SubsystemIt; SubsystemIt; ++SubsystemIt)
+	{
+		SubsystemIt->RefreshVisibilityTags();
+	}
+}
+#endif
+
 void UCommonUIVisibilitySubsystem::RefreshVisibilityTags()
 {
 	ComputedVisibilityTags = ComputeVisibilityTags();
@@ -96,7 +114,17 @@ FGameplayTagContainer UCommonUIVisibilitySubsystem::ComputeVisibilityTags() cons
 		}
 	}
 
-	ComputedTags.AppendTags(ICommonUIModule::GetSettings().GetPlatformHardwareFeatures());
+	ComputedTags.AppendTags(ICommonUIModule::GetSettings().GetPlatformTraits());
+
+#if WITH_EDITOR
+	ComputedTags.AppendTags(DebugTagsToEnable);
+	ComputedTags.RemoveTags(DebugTagsToSuppress);
+#endif
+
+	// Debug printing
+#if 0
+	UE_LOG(LogInit, Log, TEXT("UCommonUIVisibilitySubsystem::ComputeVisibilityTags() -> %s"), *ComputedTags.ToString());
+#endif
 
 	return ComputedTags;
 }
diff --git a/Engine/Plugins/Experimental/CommonUI/Source/CommonUI/Public/CommonHardwareVisibilityBorder.h b/Engine/Plugins/Experimental/CommonUI/Source/CommonUI/Public/CommonHardwareVisibilityBorder.h
--- a/Engine/Plugins/Experimental/CommonUI/Source/CommonUI/Public/CommonHardwareVisibilityBorder.h
+++ b/Engine/Plugins/Experimental/CommonUI/Source/CommonUI/Public/CommonHardwareVisibilityBorder.h
@@ -23,7 +23,7 @@ public:
 
 protected:
 
-	UPROPERTY(EditAnywhere, Category = "Visibility", meta=(Categories="Input,Hardware.Feature"))
+	UPROPERTY(EditAnywhere, Category = "Visibility", meta=(Categories="Input,Platform.Trait"))
 	FGameplayTagQuery VisibilityQuery;
 
 	UPROPERTY(EditAnywhere, Category = "Visibility")
diff --git a/Engine/Plugins/Experimental/CommonUI/Source/CommonUI/Public/CommonUISettings.h b/Engine/Plugins/Experimental/CommonUI/Source/CommonUI/Public/CommonUISettings.h
--- a/Engine/Plugins/Experimental/CommonUI/Source/CommonUI/Public/CommonUISettings.h
+++ b/Engine/Plugins/Experimental/CommonUI/Source/CommonUI/Public/CommonUISettings.h
@@ -27,9 +27,13 @@ public:
 	// Called to load CommonUISetting data, if bAutoLoadData if set to false then game code must call LoadData().
 	void LoadData();
 
+	//~UObject interface
 #if WITH_EDITOR
 	virtual void PostEditChangeProperty(struct FPropertyChangedEvent& PropertyChangedEvent) override;
+	virtual void PostReloadConfig(FProperty* PropertyThatWasLoaded) override;
+	virtual void PostInitProperties() override;
 #endif
+	//~End of UObject interface
 
 	// Called by the module startup to auto load CommonUISetting data if bAutoLoadData is true.
 	void AutoLoadData();
@@ -37,7 +41,7 @@ public:
 	UCommonUIRichTextData* GetRichTextData() const;
 	const FSlateBrush& GetDefaultThrobberBrush() const;
 	UObject* GetDefaultImageResourceObject() const;
-	const FGameplayTagContainer& GetPlatformHardwareFeatures() const { return PlatformHardwareFeatures; }
+	const FGameplayTagContainer& GetPlatformTraits() const { return PlatformTraitContainer; }
 
 private:
 
@@ -59,15 +63,21 @@ private:
 	UPROPERTY(config, EditAnywhere, Category = "RichText", meta=(AllowAbstract=false))
 	TSoftClassPtr<UCommonUIRichTextData> DefaultRichTextDataClass;
 
-	/** The Default Data for rich text to show inline icon and others. */
-	UPROPERTY(config, EditAnywhere, Category = "Visibility", meta=(Categories="Hardware.Feature", ConfigHierarchyEditable))
-	FGameplayTagContainer PlatformHardwareFeatures;
+	/** The set of traits defined per-platform (e.g., the default input mode, whether or not you can exit the application, etc...) */
+	UPROPERTY(config, EditAnywhere, Category = "Visibility", meta=(Categories="Platform.Trait", ConfigHierarchyEditable))
+	TArray<FGameplayTag> PlatformTraits;
 
 private:
 	void LoadEditorData();
+	void RebuildTraitContainer();
 
 	bool bDefaultDataLoaded;
 
+	// Merged version of PlatformTraits
+	// This is not the config property because there is no direct ini inheritance for structs
+	// (even ones like tag containers that represent a set), unlike arrays
+	FGameplayTagContainer PlatformTraitContainer;
+
 	UPROPERTY(Transient)
 	UObject* DefaultImageResourceObjectInstance;
 
diff --git a/Engine/Plugins/Experimental/CommonUI/Source/CommonUI/Public/CommonUIVisibilitySubsystem.h b/Engine/Plugins/Experimental/CommonUI/Source/CommonUI/Public/CommonUIVisibilitySubsystem.h
--- a/Engine/Plugins/Experimental/CommonUI/Source/CommonUI/Public/CommonUIVisibilitySubsystem.h
+++ b/Engine/Plugins/Experimental/CommonUI/Source/CommonUI/Public/CommonUIVisibilitySubsystem.h
@@ -36,8 +36,12 @@ public:
 	/** Get the hardware visibility tags currently in play.  These can change over time, if input mode changes, or other groups are removed/added. */
 	const FGameplayTagContainer& GetVisibilityTags() const { return ComputedVisibilityTags; }
 
-	void AddUserVisibilityCondition(const FGameplayTag& UserTag);
-	void RemoveUserVisibilityCondition(const FGameplayTag& UserTag);
+	void AddUserVisibilityCondition(const FGameplayTag UserTag);
+	void RemoveUserVisibilityCondition(const FGameplayTag UserTag);
+
+#if WITH_EDITOR
+	static void SetDebugVisibilityConditions(const FGameplayTagContainer& TagsToEnable, const FGameplayTagContainer& TagsToSuppress);
+#endif
 
 protected:
 	void RefreshVisibilityTags();
@@ -47,4 +51,8 @@ protected:
 private:
 	FGameplayTagContainer ComputedVisibilityTags;
 	FGameplayTagContainer UserVisibilityTags;
+#if WITH_EDITOR
+	static FGameplayTagContainer DebugTagsToEnable;
+	static FGameplayTagContainer DebugTagsToSuppress;
+#endif
 };
\ No newline at end of file
