CommonUI: Added some helpers to UCommonUIVisibilitySubsystem (GetChecked and HasVisibilityTag)

diff --git a/Engine/Plugins/Experimental/CommonUI/Source/CommonUI/Private/CommonUIVisibilitySubsystem.cpp b/Engine/Plugins/Experimental/CommonUI/Source/CommonUI/Private/CommonUIVisibilitySubsystem.cpp
--- a/Engine/Plugins/Experimental/CommonUI/Source/CommonUI/Private/CommonUIVisibilitySubsystem.cpp
+++ b/Engine/Plugins/Experimental/CommonUI/Source/CommonUI/Private/CommonUIVisibilitySubsystem.cpp
@@ -26,6 +26,14 @@ UCommonUIVisibilitySubsystem* UCommonUIVisibilitySubsystem::Get(const ULocalPlay
 	return LocalPlayer ? LocalPlayer->GetSubsystem<UCommonUIVisibilitySubsystem>() : nullptr;
 }
 
+UCommonUIVisibilitySubsystem* UCommonUIVisibilitySubsystem::GetChecked(const ULocalPlayer* LocalPlayer)
+{
+	check(LocalPlayer);
+	UCommonUIVisibilitySubsystem* Result = Get(LocalPlayer);
+	check(Result);
+	return Result;
+}
+
 UCommonUIVisibilitySubsystem::UCommonUIVisibilitySubsystem()
 {
 }
diff --git a/Engine/Plugins/Experimental/CommonUI/Source/CommonUI/Public/CommonUIVisibilitySubsystem.h b/Engine/Plugins/Experimental/CommonUI/Source/CommonUI/Public/CommonUIVisibilitySubsystem.h
--- a/Engine/Plugins/Experimental/CommonUI/Source/CommonUI/Public/CommonUIVisibilitySubsystem.h
+++ b/Engine/Plugins/Experimental/CommonUI/Source/CommonUI/Public/CommonUIVisibilitySubsystem.h
@@ -22,6 +22,7 @@ class COMMONUI_API UCommonUIVisibilitySubsystem : public ULocalPlayerSubsystem
 
 public:
 	static UCommonUIVisibilitySubsystem* Get(const ULocalPlayer* LocalPlayer);
+	static UCommonUIVisibilitySubsystem* GetChecked(const ULocalPlayer* LocalPlayer);
 
 	UCommonUIVisibilitySubsystem();
 	
@@ -33,8 +34,16 @@ public:
 	DECLARE_EVENT_OneParam(UCommonUIVisibilitySubsystem, FHardwareVisibilityTagsChangedEvent, UCommonUIVisibilitySubsystem*);
 	FHardwareVisibilityTagsChangedEvent OnVisibilityTagsChanged;
 
-	/** Get the hardware visibility tags currently in play.  These can change over time, if input mode changes, or other groups are removed/added. */
+	/**
+	 * Get the visibility tags currently in play (the combination of platform traits and current input tags).
+	 * These can change over time, if input mode changes, or other groups are removed/added.
+	 */
 	const FGameplayTagContainer& GetVisibilityTags() const { return ComputedVisibilityTags; }
+	
+	/* Returns true if the player currently has the specified visibility tag
+	 * (note: this value should not be cached without listening for OnVisibilityTagsChanged as it can change at runtime)
+	 */
+	bool HasVisibilityTag(const FGameplayTag VisibilityTag) const { return ComputedVisibilityTags.HasTag(VisibilityTag); }
 
 	void AddUserVisibilityCondition(const FGameplayTag UserTag);
 	void RemoveUserVisibilityCondition(const FGameplayTag UserTag);
