Resetting the SetButtonImageHeightTo editor utility, afte rusage back to 0.

diff --git a/Engine/Plugins/Experimental/CommonUI/Source/CommonInput/Private/CommonInputBaseTypes.cpp b/Engine/Plugins/Experimental/CommonUI/Source/CommonInput/Private/CommonInputBaseTypes.cpp
--- a/Engine/Plugins/Experimental/CommonUI/Source/CommonInput/Private/CommonInputBaseTypes.cpp
+++ b/Engine/Plugins/Experimental/CommonUI/Source/CommonInput/Private/CommonInputBaseTypes.cpp
@@ -202,6 +202,8 @@ void UCommonInputBaseControllerData::PostEditChangeProperty(struct FPropertyChan
 					}
 				}
 			}
+
+			SetButtonImageHeightTo = 0;
 		}
 	}
 }
@@ -331,6 +333,8 @@ bool UCommonInputPlatformSettings::SupportsInputType(ECommonInputType InputType)
 #if WITH_EDITOR
 void UCommonInputPlatformSettings::PostEditChangeProperty(struct FPropertyChangedEvent& PropertyChangedEvent)
 {
+	Super::PostEditChangeProperty(PropertyChangedEvent);
+
 	ControllerDataClasses.Reset();
 }
 #endif
