CommonUI - Making it so getting the actively simulated platform traits works.

diff --git a/Engine/Plugins/Experimental/CommonUI/Source/CommonUI/Private/CommonUISettings.cpp b/Engine/Plugins/Experimental/CommonUI/Source/CommonUI/Private/CommonUISettings.cpp
--- a/Engine/Plugins/Experimental/CommonUI/Source/CommonUI/Private/CommonUISettings.cpp
+++ b/Engine/Plugins/Experimental/CommonUI/Source/CommonUI/Private/CommonUISettings.cpp
@@ -4,6 +4,9 @@
 
 #include "Materials/MaterialInterface.h"
 #include "Materials/Material.h"
+#include "UObject/Class.h"
+#include "UObject/ReflectedTypeAccessors.h"
+#include "GameplayTagContainer.h"
 
 UCommonUISettings::UCommonUISettings(const FObjectInitializer& Initializer)
 	: Super(Initializer)
@@ -113,4 +116,36 @@ UObject* UCommonUISettings::GetDefaultImageResourceObject() const
 	ensure(bDefaultDataLoaded);
 
 	return DefaultImageResourceObjectInstance;
+}
+
+const FGameplayTagContainer& UCommonUISettings::GetPlatformTraits() const
+{
+	const FName SimulatedPlatform = UPlatformSettings::GetEditorSimulatedPlatform();
+	if (SimulatedPlatform != NAME_None)
+	{
+		if (FConfigCacheIni* PlatformIni = FConfigCacheIni::ForPlatform(SimulatedPlatform))
+		{
+			TArray<FString> Entries;
+			PlatformIni->GetArray(*GetClass()->GetPathName(), GET_MEMBER_NAME_STRING_CHECKED(UCommonUISettings, PlatformTraits), Entries, TEXT("Game"));
+
+			UScriptStruct* GameplayTagStruct = StaticStruct<FGameplayTag>();
+			TArray<FGameplayTag> OtherPlatformTraits;
+			for(FString Entry : Entries)
+			{
+				FGameplayTag Element;
+				GameplayTagStruct->ImportText(*Entry, &Element, nullptr, EPropertyPortFlags::PPF_None, nullptr, GameplayTagStruct->GetName(), true);
+				OtherPlatformTraits.Add(Element);
+			}
+
+			static FGameplayTagContainer OtherPlatformTraitContainer;
+			OtherPlatformTraitContainer.Reset();
+			for (FGameplayTag Trait : OtherPlatformTraits)
+			{
+				OtherPlatformTraitContainer.AddTag(Trait);
+			}
+			return OtherPlatformTraitContainer;
+		}
+	}
+
+	return PlatformTraitContainer;
 }
\ No newline at end of file
diff --git a/Engine/Plugins/Experimental/CommonUI/Source/CommonUI/Public/CommonUISettings.h b/Engine/Plugins/Experimental/CommonUI/Source/CommonUI/Public/CommonUISettings.h
--- a/Engine/Plugins/Experimental/CommonUI/Source/CommonUI/Public/CommonUISettings.h
+++ b/Engine/Plugins/Experimental/CommonUI/Source/CommonUI/Public/CommonUISettings.h
@@ -41,7 +41,7 @@ public:
 	UCommonUIRichTextData* GetRichTextData() const;
 	const FSlateBrush& GetDefaultThrobberBrush() const;
 	UObject* GetDefaultImageResourceObject() const;
-	const FGameplayTagContainer& GetPlatformTraits() const { return PlatformTraitContainer; }
+	const FGameplayTagContainer& GetPlatformTraits() const;
 
 private:
 
