Build - Don't try to run editor code on games.

diff --git a/Engine/Plugins/Experimental/CommonUI/Source/CommonUI/Private/CommonUISettings.cpp b/Engine/Plugins/Experimental/CommonUI/Source/CommonUI/Private/CommonUISettings.cpp
--- a/Engine/Plugins/Experimental/CommonUI/Source/CommonUI/Private/CommonUISettings.cpp
+++ b/Engine/Plugins/Experimental/CommonUI/Source/CommonUI/Private/CommonUISettings.cpp
@@ -120,6 +120,7 @@ UObject* UCommonUISettings::GetDefaultImageResourceObject() const
 
 const FGameplayTagContainer& UCommonUISettings::GetPlatformTraits() const
 {
+#if WITH_EDITOR
 	const FName SimulatedPlatform = UPlatformSettings::GetEditorSimulatedPlatform();
 	if (SimulatedPlatform != NAME_None)
 	{
@@ -146,6 +147,7 @@ const FGameplayTagContainer& UCommonUISettings::GetPlatformTraits() const
 			return OtherPlatformTraitContainer;
 		}
 	}
+#endif
 
 	return PlatformTraitContainer;
 }
\ No newline at end of file
