CommonUI - CommonButtons will treat unhandled double click input as a mouse down, so that buttons that require rapid clicking feel good, and don't lose every other click on rapid clicking.

diff --git a/Engine/Plugins/Experimental/CommonUI/Source/CommonUI/Private/CommonButtonTypes.cpp b/Engine/Plugins/Experimental/CommonUI/Source/CommonUI/Private/CommonButtonTypes.cpp
--- a/Engine/Plugins/Experimental/CommonUI/Source/CommonUI/Private/CommonButtonTypes.cpp
+++ b/Engine/Plugins/Experimental/CommonUI/Source/CommonUI/Private/CommonButtonTypes.cpp
@@ -11,7 +11,7 @@ FReply SCommonButton::OnMouseButtonDown(const FGeometry& MyGeometry, const FPoin
 	return IsInteractable() ? SButton::OnMouseButtonDown(MyGeometry, MouseEvent) : FReply::Handled();
 }
 
-FReply SCommonButton::OnMouseButtonDoubleClick(const FGeometry& InMyGeometry, const FPointerEvent& InMouseEvent)
+FReply SCommonButton::OnMouseButtonDoubleClick(const FGeometry& InMyGeometry, const FPointerEvent& InPointerEvent)
 {
 	if (!IsInteractable())
 	{
@@ -27,10 +27,22 @@ FReply SCommonButton::OnMouseButtonDoubleClick(const FGeometry& InMyGeometry, co
 		}
 	}
 
-	return SButton::OnMouseButtonDoubleClick(InMyGeometry, InMouseEvent);
+	// The default button behavior 'ignores' double click, which means you get,
+	// down -> up, double-click -> up
+	// which can make the input feel like it's being lost when players are clicking
+	// a button over and over.  So if your button does not handle the double click
+	// specifically, we'll treat the double click as a mouse up, and do whatever
+	// we would normally do based on button configuration.
+	FReply Reply = OnMouseButtonDown(InMyGeometry, InPointerEvent);
+	if (Reply.IsEventHandled())
+	{
+		return Reply;
+	}
+
+	return SButton::OnMouseButtonDoubleClick(InMyGeometry, InPointerEvent);
 }
 
-FReply SCommonButton::OnMouseButtonUp(const FGeometry& MyGeometry, const FPointerEvent& MouseEvent)
+FReply SCommonButton::OnMouseButtonUp(const FGeometry& MyGeometry, const FPointerEvent& InPointerEvent)
 {
 	FReply Reply = FReply::Handled();
 	if (!IsInteractable())
@@ -45,25 +57,25 @@ FReply SCommonButton::OnMouseButtonUp(const FGeometry& MyGeometry, const FPointe
 	}
 	else
 	{
-		Reply = SButton::OnMouseButtonUp(MyGeometry, MouseEvent);
+		Reply = SButton::OnMouseButtonUp(MyGeometry, InPointerEvent);
 	}
 
 	return Reply;
 }
 
-void SCommonButton::OnMouseEnter(const FGeometry& MyGeometry, const FPointerEvent& MouseEvent)
+void SCommonButton::OnMouseEnter(const FGeometry& MyGeometry, const FPointerEvent& InPointerEvent)
 {
-	if (!MouseEvent.IsTouchEvent())
+	if (!InPointerEvent.IsTouchEvent())
 	{
 		bHovered = true;
 		SetHover(bHovered && bIsInteractionEnabled);
-		SButton::OnMouseEnter(MyGeometry, MouseEvent);
+		SButton::OnMouseEnter(MyGeometry, InPointerEvent);
 	}
 }
 
-void SCommonButton::OnMouseLeave(const FPointerEvent& MouseEvent)
+void SCommonButton::OnMouseLeave(const FPointerEvent& InPointerEvent)
 {
-	if (MouseEvent.IsTouchEvent())
+	if (InPointerEvent.IsTouchEvent())
 	{
 		if (HasMouseCapture())
 		{
@@ -74,7 +86,7 @@ void SCommonButton::OnMouseLeave(const FPointerEvent& MouseEvent)
 	{
 		bHovered = false;
 		SetHover(false);
-		SButton::OnMouseLeave(MouseEvent);
+		SButton::OnMouseLeave(InPointerEvent);
 	}
 }
 
