CommonUI: Add limited support for gamepad type autodetection based on any active FInputDeviceScope
- The detection criteria can be defined in gamepad assets
- Limited because XInput makes it totally impossible to distinguish between different controller types (other plugin sources will still work), so to preserve the ability to set a preference in project settings the logic prefers to stick with an existing match rather than switching to something else, but otherwise will prefer the first match found in the controller asset list

diff --git a/Engine/Plugins/Experimental/CommonUI/Source/CommonInput/CommonInput.Build.cs b/Engine/Plugins/Experimental/CommonUI/Source/CommonInput/CommonInput.Build.cs
--- a/Engine/Plugins/Experimental/CommonUI/Source/CommonInput/CommonInput.Build.cs
+++ b/Engine/Plugins/Experimental/CommonUI/Source/CommonInput/CommonInput.Build.cs
@@ -22,6 +22,7 @@ public class CommonInput : ModuleRules
 			{
                 "SlateCore",
 				"Slate",
+				"ApplicationCore",
                 "EngineSettings"
             }
 		);
diff --git a/Engine/Plugins/Experimental/CommonUI/Source/CommonInput/Private/CommonInputBaseTypes.cpp b/Engine/Plugins/Experimental/CommonUI/Source/CommonInput/Private/CommonInputBaseTypes.cpp
--- a/Engine/Plugins/Experimental/CommonUI/Source/CommonInput/Private/CommonInputBaseTypes.cpp
+++ b/Engine/Plugins/Experimental/CommonUI/Source/CommonInput/Private/CommonInputBaseTypes.cpp
@@ -269,7 +269,7 @@ void UCommonInputPlatformSettings::InitializePlatformDefaults()
 	DefaultGamepadName = PlatformName;
 }
 
-bool UCommonInputPlatformSettings::TryGetInputBrush(FSlateBrush& OutBrush, FKey Key, ECommonInputType InputType, const FName& GamepadName) const
+bool UCommonInputPlatformSettings::TryGetInputBrush(FSlateBrush& OutBrush, FKey Key, ECommonInputType InputType, const FName GamepadName) const
 {
 	InitializeControllerData();
 
@@ -288,7 +288,7 @@ bool UCommonInputPlatformSettings::TryGetInputBrush(FSlateBrush& OutBrush, FKey
 	return false;
 }
 
-bool UCommonInputPlatformSettings::TryGetInputBrush(FSlateBrush& OutBrush, const TArray<FKey>& Keys, ECommonInputType InputType, const FName& GamepadName) const
+bool UCommonInputPlatformSettings::TryGetInputBrush(FSlateBrush& OutBrush, const TArray<FKey>& Keys, ECommonInputType InputType, const FName GamepadName) const
 {
 	InitializeControllerData();
 
@@ -307,6 +307,51 @@ bool UCommonInputPlatformSettings::TryGetInputBrush(FSlateBrush& OutBrush, const
 	return false;
 }
 
+FName UCommonInputPlatformSettings::GetBestGamepadNameForHardware(FName CurrentGamepadName, FName InputDeviceName, const FString& HardwareDeviceIdentifier)
+{
+	FName FirstMatch = NAME_None;
+
+	// This is far more complicated than it should be because XInput exposes no information about device type,
+	// so we want to be 'sticky', only switching to an Xbox controller if you don't already have one selected
+	// and otherwise conserving the player UI chosen choice
+	if (ControllerDataClasses.Num() > 0)
+	{
+		for (const TSubclassOf<UCommonInputBaseControllerData>& ControllerDataPtr : ControllerDataClasses)
+		{
+			if (const UCommonInputBaseControllerData* DefaultControllerData = ControllerDataPtr.GetDefaultObject())
+			{
+				bool bThisEntryMatches = false;
+				for (const FInputDeviceIdentifierPair& Pair : DefaultControllerData->GamepadHardwareIdMapping)
+				{
+					if ((Pair.InputDeviceName == InputDeviceName) && (Pair.HardwareDeviceIdentifier == HardwareDeviceIdentifier))
+					{
+						bThisEntryMatches = true;
+						break;
+					}
+				}
+
+				if (bThisEntryMatches)
+				{
+					if (CurrentGamepadName == DefaultControllerData->GamepadName)
+					{
+						// Preferentially conserve the existing setting
+						return CurrentGamepadName;
+					}
+
+					if (FirstMatch == NAME_None)
+					{
+						// Record the first match, which we'll use if the existing one doesn't work
+						FirstMatch = DefaultControllerData->GamepadName;
+					}
+				}
+			}
+
+		}
+	}
+
+	return FirstMatch.IsNone() ? CurrentGamepadName : FirstMatch;
+}
+
 bool UCommonInputPlatformSettings::SupportsInputType(ECommonInputType InputType) const
 {
 	switch (InputType)
@@ -339,29 +384,6 @@ void UCommonInputPlatformSettings::PostEditChangeProperty(struct FPropertyChange
 }
 #endif
 
-
-
-
-
-
-
-
-
-
-
-
-
-
-
-
-
-
-
-
-
-
-
-
 bool FCommonInputPlatformBaseData::TryGetInputBrush(FSlateBrush& OutBrush, FKey Key, ECommonInputType InputType, const FName& GamepadName) const
 {
 	if (ControllerDataClasses.Num() > 0)
diff --git a/Engine/Plugins/Experimental/CommonUI/Source/CommonInput/Private/CommonInputSubsystem.cpp b/Engine/Plugins/Experimental/CommonUI/Source/CommonInput/Private/CommonInputSubsystem.cpp
--- a/Engine/Plugins/Experimental/CommonUI/Source/CommonInput/Private/CommonInputSubsystem.cpp
+++ b/Engine/Plugins/Experimental/CommonUI/Source/CommonInput/Private/CommonInputSubsystem.cpp
@@ -203,6 +203,20 @@ private:
 #endif
 
 		InputSubsystem.SetCurrentInputType(InputMethod);
+
+		if (UCommonInputPlatformSettings::Get()->CanChangeGamepadType())
+		{
+			if (const FInputDeviceScope* DeviceScope = FInputDeviceScope::GetCurrent())
+			{
+				const FName GamepadInputType = InputSubsystem.GetCurrentGamepadName();
+				const FName BestGamepadType = UCommonInputPlatformSettings::Get()->GetBestGamepadNameForHardware(GamepadInputType, DeviceScope->InputDeviceName, DeviceScope->HardwareDeviceIdentifier);
+				if (BestGamepadType != GamepadInputType)
+				{
+					UE_LOG(LogCommonInput, Log, TEXT("UCommonInputSubsystem: Autodetect changed GamepadInputType to %s"), *BestGamepadType.ToString());
+					InputSubsystem.SetGamepadInputType(BestGamepadType);
+				}
+			}
+		}
 	}
 
 	ECommonInputType GetInputType(const FKey& Key)
@@ -267,6 +281,7 @@ UCommonInputSubsystem* UCommonInputSubsystem::Get(const ULocalPlayer* LocalPlaye
 
 UCommonInputSubsystem::UCommonInputSubsystem()
 {
+	//@TODO: Project setting (should be done as another config var any platform can set tho)?
 //	// Uncomment this if we want mobile platforms to default to gamepad when gamepads are enabled.
 // #if PLATFORM_IOS
 //     bool bAllowControllers = false;
@@ -444,7 +459,7 @@ bool UCommonInputSubsystem::CheckForInputMethodThrashing(ECommonInputType NewInp
 
 void UCommonInputSubsystem::SetCurrentInputType(ECommonInputType NewInputType)
 {
-	if (LastInputType != NewInputType && PlatformSupportsInputType(NewInputType))
+	if ((LastInputType != NewInputType) && PlatformSupportsInputType(NewInputType))
 	{
 		CheckForInputMethodThrashing(NewInputType);
 	
@@ -496,19 +511,16 @@ const FName UCommonInputSubsystem::GetCurrentGamepadName() const
 	return GamepadInputType;
 }
  
-void UCommonInputSubsystem::SetGamepadInputType(const FName& InGamepadInputType)
+void UCommonInputSubsystem::SetGamepadInputType(const FName InGamepadInputType)
 {
-	if (!UCommonInputPlatformSettings::Get()->CanChangeGamepadType())
+	if (ensure(UCommonInputPlatformSettings::Get()->CanChangeGamepadType()))
 	{
-		ensure(false);	// should not be called on the console platforms
-		return;
-	}
-
-	GamepadInputType = InGamepadInputType;
+		GamepadInputType = InGamepadInputType;
 
-	// Send out notifications so we update our buttons
-	//BroadcastLastInputDeviceChanged();
-	BroadcastInputMethodChanged();
+		// Send out notifications so we update our buttons
+		//BroadcastLastInputDeviceChanged();
+		BroadcastInputMethodChanged();
+	}
 }
 
 bool UCommonInputSubsystem::IsUsingPointerInput() const
diff --git a/Engine/Plugins/Experimental/CommonUI/Source/CommonInput/Public/CommonInputBaseTypes.h b/Engine/Plugins/Experimental/CommonUI/Source/CommonInput/Public/CommonInputBaseTypes.h
--- a/Engine/Plugins/Experimental/CommonUI/Source/CommonInput/Public/CommonInputBaseTypes.h
+++ b/Engine/Plugins/Experimental/CommonUI/Source/CommonInput/Public/CommonInputBaseTypes.h
@@ -74,6 +74,18 @@ public:
 	FSlateBrush KeyBrush;
 };
 
+USTRUCT()
+struct FInputDeviceIdentifierPair
+{
+	GENERATED_BODY()
+
+	UPROPERTY(EditDefaultsOnly, Category = "Gamepad")
+	FName InputDeviceName;
+
+	UPROPERTY(EditDefaultsOnly, Category = "Gamepad")
+	FString HardwareDeviceIdentifier;
+};
+
 /* Derive from this class to store the Input data. It is referenced in the Common Input Settings, found in the project settings UI. */
 UCLASS(Abstract, Blueprintable, ClassGroup = Input, meta = (Category = "Common Input"))
 class COMMONINPUT_API UCommonUIInputData : public UObject
@@ -121,6 +133,9 @@ public:
 	UPROPERTY(EditDefaultsOnly, Category = "Gamepad", meta = (EditCondition = "InputType == ECommonInputType::Gamepad"))
 	FText GamepadDisplayName;
 
+	UPROPERTY(EditDefaultsOnly, Category = "Gamepad", meta=(EditCondition="InputType == ECommonInputType::Gamepad"))
+	TArray<FInputDeviceIdentifierPair> GamepadHardwareIdMapping;
+
 	UPROPERTY(EditDefaultsOnly, Category = "Display")
 	TSoftObjectPtr<UTexture2D> ControllerTexture;
 
@@ -159,8 +174,10 @@ public:
 		return UPlatformSettings::GetSettingsForPlatform<UCommonInputPlatformSettings>();
 	}
 
-	bool TryGetInputBrush(FSlateBrush& OutBrush, FKey Key, ECommonInputType InputType, const FName& GamepadName) const;
-	bool TryGetInputBrush(FSlateBrush& OutBrush, const TArray<FKey>& Keys, ECommonInputType InputType,  const FName& GamepadName) const;
+	bool TryGetInputBrush(FSlateBrush& OutBrush, FKey Key, ECommonInputType InputType, const FName GamepadName) const;
+	bool TryGetInputBrush(FSlateBrush& OutBrush, const TArray<FKey>& Keys, ECommonInputType InputType, const FName GamepadName) const;
+
+	FName GetBestGamepadNameForHardware(FName CurrentGamepadName, FName InputDeviceName, const FString& HardwareDeviceIdentifier);
 
 	ECommonInputType GetDefaultInputType() const
 	{
diff --git a/Engine/Plugins/Experimental/CommonUI/Source/CommonInput/Public/CommonInputSubsystem.h b/Engine/Plugins/Experimental/CommonUI/Source/CommonInput/Public/CommonInputSubsystem.h
--- a/Engine/Plugins/Experimental/CommonUI/Source/CommonInput/Public/CommonInputSubsystem.h
+++ b/Engine/Plugins/Experimental/CommonUI/Source/CommonInput/Public/CommonInputSubsystem.h
@@ -54,7 +54,7 @@ public:
 	const FName GetCurrentGamepadName() const;
 
 	UFUNCTION(BlueprintCallable, Category = CommonInputSubsystem)
-	void SetGamepadInputType(const FName& InGamepadInputType);
+	void SetGamepadInputType(const FName InGamepadInputType);
 
 	UFUNCTION(BlueprintCallable, Category = CommonInputSubsystem)
 	bool IsUsingPointerInput() const;
