CommonUI: Slight optimization to input autodetection to avoid doing it when the most recently seen FInputDeviceScope IDs haven't changed

diff --git a/Engine/Plugins/Experimental/CommonUI/Source/CommonInput/Private/CommonInputSubsystem.cpp b/Engine/Plugins/Experimental/CommonUI/Source/CommonInput/Private/CommonInputSubsystem.cpp
--- a/Engine/Plugins/Experimental/CommonUI/Source/CommonInput/Private/CommonInputSubsystem.cpp
+++ b/Engine/Plugins/Experimental/CommonUI/Source/CommonInput/Private/CommonInputSubsystem.cpp
@@ -204,16 +204,23 @@ private:
 
 		InputSubsystem.SetCurrentInputType(InputMethod);
 
-		if (UCommonInputPlatformSettings::Get()->CanChangeGamepadType())
+		// Try to auto-detect the type of gamepad
+		if ((InputMethod == ECommonInputType::Gamepad) && UCommonInputPlatformSettings::Get()->CanChangeGamepadType())
 		{
 			if (const FInputDeviceScope* DeviceScope = FInputDeviceScope::GetCurrent())
 			{
-				const FName GamepadInputType = InputSubsystem.GetCurrentGamepadName();
-				const FName BestGamepadType = UCommonInputPlatformSettings::Get()->GetBestGamepadNameForHardware(GamepadInputType, DeviceScope->InputDeviceName, DeviceScope->HardwareDeviceIdentifier);
-				if (BestGamepadType != GamepadInputType)
+				if ((DeviceScope->InputDeviceName != LastSeenGamepadInputDeviceName) || (DeviceScope->HardwareDeviceIdentifier != LastSeenGamepadHardwareDeviceIdentifier))
 				{
-					UE_LOG(LogCommonInput, Log, TEXT("UCommonInputSubsystem: Autodetect changed GamepadInputType to %s"), *BestGamepadType.ToString());
-					InputSubsystem.SetGamepadInputType(BestGamepadType);
+					LastSeenGamepadInputDeviceName = DeviceScope->InputDeviceName;
+					LastSeenGamepadHardwareDeviceIdentifier = DeviceScope->HardwareDeviceIdentifier;
+
+					const FName GamepadInputType = InputSubsystem.GetCurrentGamepadName();
+					const FName BestGamepadType = UCommonInputPlatformSettings::Get()->GetBestGamepadNameForHardware(GamepadInputType, DeviceScope->InputDeviceName, DeviceScope->HardwareDeviceIdentifier);
+					if (BestGamepadType != GamepadInputType)
+					{
+						UE_LOG(LogCommonInput, Log, TEXT("UCommonInputSubsystem: Autodetect changed GamepadInputType to %s"), *BestGamepadType.ToString());
+						InputSubsystem.SetGamepadInputType(BestGamepadType);
+					}
 				}
 			}
 		}
@@ -258,6 +265,9 @@ private:
 	// The reasons we might be filtering input right now.
 	TMap<FName, bool> FilterInputTypeWithReasons[(uint8)ECommonInputType::Count];
 
+	FName LastSeenGamepadInputDeviceName;
+	FString LastSeenGamepadHardwareDeviceIdentifier;
+
 	friend class UCommonInputSubsystem;
 };
 
