Adding additional text options for gamepad controls

diff --git a/Engine/Plugins/Experimental/CommonUI/Source/CommonInput/Public/CommonInputBaseTypes.h b/Engine/Plugins/Experimental/CommonUI/Source/CommonInput/Public/CommonInputBaseTypes.h
--- a/Engine/Plugins/Experimental/CommonUI/Source/CommonInput/Public/CommonInputBaseTypes.h
+++ b/Engine/Plugins/Experimental/CommonUI/Source/CommonInput/Public/CommonInputBaseTypes.h
@@ -133,6 +133,12 @@ public:
 	UPROPERTY(EditDefaultsOnly, Category = "Gamepad", meta = (EditCondition = "InputType == ECommonInputType::Gamepad"))
 	FText GamepadDisplayName;
 
+	UPROPERTY(EditDefaultsOnly, Category = "Gamepad", meta = (EditCondition = "InputType == ECommonInputType::Gamepad"))
+	FText GamepadCategory;
+
+	UPROPERTY(EditDefaultsOnly, Category = "Gamepad", meta = (EditCondition = "InputType == ECommonInputType::Gamepad"))
+	FText GamepadPlatformName;
+
 	UPROPERTY(EditDefaultsOnly, Category = "Gamepad", meta=(EditCondition="InputType == ECommonInputType::Gamepad"))
 	TArray<FInputDeviceIdentifierPair> GamepadHardwareIdMapping;
 
