Fixing some uninitialized member warnings

diff --git a/Engine/Plugins/Bridge/Source/MegascansPlugin/Private/MSAssetImportData.h b/Engine/Plugins/Bridge/Source/MegascansPlugin/Private/MSAssetImportData.h
--- a/Engine/Plugins/Bridge/Source/MegascansPlugin/Private/MSAssetImportData.h
+++ b/Engine/Plugins/Bridge/Source/MegascansPlugin/Private/MSAssetImportData.h
@@ -170,7 +170,7 @@ public:
 	UPROPERTY()
 		FString assetSubType;
 	UPROPERTY()
-		int8 assetTier;
+		int8 assetTier = 0;
 	UPROPERTY()
 		FString assetRootPath;
 	UPROPERTY()
diff --git a/Engine/Plugins/Experimental/CommonUI/Source/CommonUI/Public/CommonButtonBase.h b/Engine/Plugins/Experimental/CommonUI/Source/CommonUI/Public/CommonButtonBase.h
--- a/Engine/Plugins/Experimental/CommonUI/Source/CommonUI/Public/CommonButtonBase.h
+++ b/Engine/Plugins/Experimental/CommonUI/Source/CommonUI/Public/CommonButtonBase.h
@@ -32,7 +32,7 @@ struct FCommonButtonStyleOptionalSlateSound
 
 public:
 	UPROPERTY(EditAnywhere, BlueprintReadWrite, Category = "Properties")
-	bool bHasSound;
+	bool bHasSound = false;
 
 	UPROPERTY(EditAnywhere, BlueprintReadWrite, Category = "Properties", meta = (EditCondition = "bHasSound"))
 	FSlateSound Sound;
