Non-unity fixes

diff --git a/Engine/Plugins/Experimental/CommonUI/Source/CommonUI/Private/CommonUISettings.cpp b/Engine/Plugins/Experimental/CommonUI/Source/CommonUI/Private/CommonUISettings.cpp
--- a/Engine/Plugins/Experimental/CommonUI/Source/CommonUI/Private/CommonUISettings.cpp
+++ b/Engine/Plugins/Experimental/CommonUI/Source/CommonUI/Private/CommonUISettings.cpp
@@ -7,6 +7,7 @@
 #include "UObject/Class.h"
 #include "UObject/ReflectedTypeAccessors.h"
 #include "GameplayTagContainer.h"
+#include "Misc/ConfigCacheIni.h"
 
 UCommonUISettings::UCommonUISettings(const FObjectInitializer& Initializer)
 	: Super(Initializer)
diff --git a/Engine/Plugins/Experimental/CommonUI/Source/CommonUI/Private/CommonUIVisibilitySubsystem.cpp b/Engine/Plugins/Experimental/CommonUI/Source/CommonUI/Private/CommonUIVisibilitySubsystem.cpp
--- a/Engine/Plugins/Experimental/CommonUI/Source/CommonUI/Private/CommonUIVisibilitySubsystem.cpp
+++ b/Engine/Plugins/Experimental/CommonUI/Source/CommonUI/Private/CommonUIVisibilitySubsystem.cpp
@@ -9,6 +9,7 @@
 #include "Engine/GameInstance.h"
 #include "CommonInputSubsystem.h"
 #include "UObject/UObjectHash.h"
+#include "UObject/UObjectIterator.h"
 #include "ICommonUIModule.h"
 #include "CommonUISettings.h"
 
diff --git a/Engine/Source/Developer/LogVisualizer/Private/SVisualLoggerToolbar.cpp b/Engine/Source/Developer/LogVisualizer/Private/SVisualLoggerToolbar.cpp
--- a/Engine/Source/Developer/LogVisualizer/Private/SVisualLoggerToolbar.cpp
+++ b/Engine/Source/Developer/LogVisualizer/Private/SVisualLoggerToolbar.cpp
@@ -1,6 +1,7 @@
 // Copyright Epic Games, Inc. All Rights Reserved.
 
 #include "SVisualLoggerToolbar.h"
+#include "EditorStyleSet.h"
 #include "Framework/MultiBox/MultiBoxDefs.h"
 #include "Framework/MultiBox/MultiBoxBuilder.h"
 #include "LogVisualizerStyle.h"
diff --git a/Engine/Source/Editor/Persona/Private/SAnimationBlendSpace.h b/Engine/Source/Editor/Persona/Private/SAnimationBlendSpace.h
--- a/Engine/Source/Editor/Persona/Private/SAnimationBlendSpace.h
+++ b/Engine/Source/Editor/Persona/Private/SAnimationBlendSpace.h
@@ -14,6 +14,7 @@
 #include "Misc/NotifyHook.h"
 #include "EditorUndoClient.h"
 
+class FScopedTransaction;
 class SBlendSpaceGridWidget;
 
 class SBlendSpaceEditor : public SAnimEditorBase, public FNotifyHook, public FSelfRegisteringEditorUndoClient
