Adding CommonUI functionality to better support blueprint-only workflows.
----
Adding functions to bind visibility to widget activation.
Exposing GetFocusTarget on activatable widgets to BP.
Add priority and duplication removal to nav-bar.
Fix warning when using default back actions.
Since we no longer use CommonActivatablePanel, always recalculate input bindings on widget activation / deactivation.
Check for tab group before deleting for BP previews.

diff --git a/Engine/Plugins/Experimental/CommonUI/Source/CommonUI/Private/CommonActivatableWidget.cpp b/Engine/Plugins/Experimental/CommonUI/Source/CommonUI/Private/CommonActivatableWidget.cpp
--- a/Engine/Plugins/Experimental/CommonUI/Source/CommonUI/Private/CommonActivatableWidget.cpp
+++ b/Engine/Plugins/Experimental/CommonUI/Source/CommonUI/Private/CommonActivatableWidget.cpp
@@ -16,7 +16,7 @@ void UCommonActivatableWidget::NativeConstruct()
 		FBindUIActionArgs BindArgs(ICommonInputModule::GetSettings().GetDefaultBackAction(), FSimpleDelegate::CreateUObject(this, &UCommonActivatableWidget::HandleBackAction));
 		BindArgs.bDisplayInActionBar = bIsBackActionDisplayedInActionBar;
 
-		RegisterUIActionBinding(BindArgs);
+		DefaultBackActionHandle = RegisterUIActionBinding(BindArgs);
 	}
 
 	if (bAutoActivate)
@@ -35,6 +35,8 @@ void UCommonActivatableWidget::NativeDestruct()
 		DeactivateWidget();
 	}
 	Super::NativeDestruct();
+
+	DefaultBackActionHandle.Unregister();
 }
 
 UWidget* UCommonActivatableWidget::GetDesiredFocusTarget() const
@@ -85,6 +87,25 @@ void UCommonActivatableWidget::DeactivateWidget()
 	}
 }
 
+void UCommonActivatableWidget::SetBindVisibilities(ESlateVisibility OnActivatedVisibility, ESlateVisibility OnDeactivatedVisibility, bool bInAllActive)
+{
+	ActivatedBindVisibility = OnActivatedVisibility;
+	DeactivatedBindVisibility = OnDeactivatedVisibility;
+	bAllActive = bInAllActive;
+}
+
+void UCommonActivatableWidget::BindVisibilityToActivation(UCommonActivatableWidget* ActivatableWidget)
+{
+	if (ActivatableWidget && !VisibilityBoundWidgets.Contains(ActivatableWidget))
+	{
+		VisibilityBoundWidgets.Add(ActivatableWidget);
+		ActivatableWidget->OnActivated().AddUObject(this, &UCommonActivatableWidget::HandleVisibilityBoundWidgetActivations);
+		ActivatableWidget->OnDeactivated().AddUObject(this, &UCommonActivatableWidget::HandleVisibilityBoundWidgetActivations);
+
+		HandleVisibilityBoundWidgetActivations();
+	}
+}
+
 void UCommonActivatableWidget::InternalProcessDeactivation()
 {
 	bIsActive = false;
@@ -162,6 +183,38 @@ void UCommonActivatableWidget::HandleBackAction()
 	NativeOnHandleBackAction();
 }
 
+void UCommonActivatableWidget::HandleVisibilityBoundWidgetActivations()
+{
+	ESlateVisibility OldDeactivatedVisibility = DeactivatedBindVisibility;
+	OldDeactivatedVisibility = bSetVisibilityOnActivated && IsActivated() ? ActivatedVisibility : OldDeactivatedVisibility;
+	OldDeactivatedVisibility = bSetVisibilityOnDeactivated  && !IsActivated() ? DeactivatedVisibility : OldDeactivatedVisibility;
+
+	for (const TWeakObjectPtr<UCommonActivatableWidget>& VisibilityBoundWidget : VisibilityBoundWidgets)
+	{
+		if (VisibilityBoundWidget.IsValid())
+		{
+			if (bAllActive)
+			{
+				if (!VisibilityBoundWidget->IsActivated())
+				{
+					SetVisibility(OldDeactivatedVisibility);
+					return;
+				}
+			}
+			else 
+			{
+				if (VisibilityBoundWidget->IsActivated())
+				{
+					SetVisibility(ActivatedBindVisibility);
+					return;
+				}
+			}
+		}
+	}
+
+	SetVisibility(bAllActive ? ActivatedBindVisibility : OldDeactivatedVisibility);
+}
+
 void UCommonActivatableWidget::Reset()
 {
 	bIsActive = false;
diff --git a/Engine/Plugins/Experimental/CommonUI/Source/CommonUI/Private/CommonTabListWidgetBase.cpp b/Engine/Plugins/Experimental/CommonUI/Source/CommonUI/Private/CommonTabListWidgetBase.cpp
--- a/Engine/Plugins/Experimental/CommonUI/Source/CommonUI/Private/CommonTabListWidgetBase.cpp
+++ b/Engine/Plugins/Experimental/CommonUI/Source/CommonUI/Private/CommonTabListWidgetBase.cpp
@@ -322,7 +322,10 @@ void UCommonTabListWidgetBase::NativeDestruct()
 
 	ActiveTabID = NAME_None;
 	RemoveAllTabs();
-	TabButtonGroup->RemoveAll();
+	if (TabButtonGroup)
+	{
+		TabButtonGroup->RemoveAll();
+	}
 }
 
 void UCommonTabListWidgetBase::HandlePreLinkedSwitcherChanged()
diff --git a/Engine/Plugins/Experimental/CommonUI/Source/CommonUI/Private/Input/CommonBoundActionBar.cpp b/Engine/Plugins/Experimental/CommonUI/Source/CommonUI/Private/Input/CommonBoundActionBar.cpp
--- a/Engine/Plugins/Experimental/CommonUI/Source/CommonUI/Private/Input/CommonBoundActionBar.cpp
+++ b/Engine/Plugins/Experimental/CommonUI/Source/CommonUI/Private/Input/CommonBoundActionBar.cpp
@@ -135,7 +135,8 @@ void UCommonBoundActionBar::HandleDeferredDisplayUpdate()
 					const ECommonInputType PlayerInputType = InputSubsystem.GetCurrentInputType();
 					const FName& PlayerGamepadName = InputSubsystem.GetCurrentGamepadName();
 
-					TArray<FUIActionBindingHandle> FilteredBindings = ActionRouter->GatherActiveBindings().FilterByPredicate([PlayerInputType, PlayerGamepadName](const FUIActionBindingHandle& Handle)
+					TSet<FName> AcceptedBindings;
+					TArray<FUIActionBindingHandle> FilteredBindings = ActionRouter->GatherActiveBindings().FilterByPredicate([PlayerInputType, PlayerGamepadName, &AcceptedBindings](const FUIActionBindingHandle& Handle) mutable
 						{
 							if (TSharedPtr<FUIActionBinding> Binding = FUIActionBinding::FindBinding(Handle))
 							{
@@ -156,7 +157,9 @@ void UCommonBoundActionBar::HandleDeferredDisplayUpdate()
 									return false; //@todo(josh.gross) - allow non-legacy bindings
 								}
 
-								return true;
+								bool bAlreadyAccepted = false;
+								AcceptedBindings.Add(Binding->ActionName, &bAlreadyAccepted);
+								return !bAlreadyAccepted;
 							}
 
 							return false;
@@ -200,6 +203,10 @@ void UCommonBoundActionBar::HandleDeferredDisplayUpdate()
 								{
 									return false;
 								}
+								else if (LegacyDataA->NavBarPriority != LegacyDataB->NavBarPriority)
+								{
+									return LegacyDataA->NavBarPriority < LegacyDataB->NavBarPriority;
+								}
 							}
 
 							return GetTypeHash(BindingA->Handle) < GetTypeHash(BindingB->Handle);
diff --git a/Engine/Plugins/Experimental/CommonUI/Source/CommonUI/Private/Input/UIActionRouterTypes.cpp b/Engine/Plugins/Experimental/CommonUI/Source/CommonUI/Private/Input/UIActionRouterTypes.cpp
--- a/Engine/Plugins/Experimental/CommonUI/Source/CommonUI/Private/Input/UIActionRouterTypes.cpp
+++ b/Engine/Plugins/Experimental/CommonUI/Source/CommonUI/Private/Input/UIActionRouterTypes.cpp
@@ -1127,6 +1127,10 @@ bool FActivatableTreeRoot::UpdateLeafmostActiveNode(FActivatableTreeNodePtr Base
 
 	// Starting from the provided candidate, find the leafmost active node (which may well just be the candidate itself)
 	FActivatableTreeNodePtr NewLeafmostNode = BaseCandidateNode ? FindLeafmostActiveChild(BaseCandidateNode.ToSharedRef()) : nullptr;
+
+	// Always trigger on valid canidates to allow to bound actions to update
+	OnLeafmostActiveNodeChanged.ExecuteIfBound();
+
 	if (NewLeafmostNode != LeafmostActiveNode)
 	{
 		if (LeafmostActiveNode.IsValid())
@@ -1151,7 +1155,6 @@ bool FActivatableTreeRoot::UpdateLeafmostActiveNode(FActivatableTreeNodePtr Base
 			}
 		}
 
-		OnLeafmostActiveNodeChanged.ExecuteIfBound();
 		return true;
 	}
 	return false;
diff --git a/Engine/Plugins/Experimental/CommonUI/Source/CommonUI/Public/CommonActivatableWidget.h b/Engine/Plugins/Experimental/CommonUI/Source/CommonUI/Public/CommonActivatableWidget.h
--- a/Engine/Plugins/Experimental/CommonUI/Source/CommonUI/Public/CommonActivatableWidget.h
+++ b/Engine/Plugins/Experimental/CommonUI/Source/CommonUI/Public/CommonActivatableWidget.h
@@ -40,15 +40,35 @@ public:
 
 	UFUNCTION(BlueprintCallable, Category = ActivatableWidget)
 	void DeactivateWidget();
+	
+	/**
+	 * Visibilities to use for when bound widgets in BindVisibilityToActivation are activated.
+	 *
+	 * @param	OnActivatedVisibility	- Visibility for when bound widgets are active
+	 * @param	OnDeactivatedVisibility - Visibility for when bound widgets are not active, not used if this widget has activation / deactivation visibilities
+	 * @param	bInAllActive			- True if we should switch to activated visibility only when all bound widgets are active
+	 */
+	UFUNCTION(BlueprintCallable, Category = ActivatableWidget, meta=(BlueprintProtected = "true"))
+	void SetBindVisibilities(ESlateVisibility OnActivatedVisibility, ESlateVisibility OnDeactivatedVisibility, UPARAM(DisplayName = "All Active") bool bInAllActive);
+
+	/**
+	 * Bind our visibility to the activation of another widget, useful for making mouse collisions behave similiar to console navigation w.r.t activation
+	 * Will immediately update visibility based on the bound widget activation & visibilites set by SetBindVisibilities.
+	 *
+	 * @param	ActivatableWidget		- The widget whose activation / deactivation will modify our visibility 
+	 */
+	UFUNCTION(BlueprintCallable, Category = ActivatableWidget)
+	void BindVisibilityToActivation(UCommonActivatableWidget* ActivatableWidget);
+
+	/** Returns the desired widget to focus when this Widget Activates. */
+	UFUNCTION(BlueprintCallable, Category = ActivatableWidget)
+	UWidget* GetDesiredFocusTarget() const;
 
 	FSimpleMulticastDelegate& OnActivated() const { return OnActivatedEvent; }
 	FSimpleMulticastDelegate& OnDeactivated() const { return OnDeactivatedEvent; }
 
 //COMMONUI_SCOPE:
 public:
-
-	/** Returns the desired widget to focus when this Widget Activates. */
-	UWidget* GetDesiredFocusTarget() const;
 	
 	virtual TOptional<FUICameraConfig> GetDesiredCameraConfig() const;
 	/**
@@ -154,6 +174,9 @@ protected:
 	bool bAutoRestoreFocus = false;
 	
 private:
+	/** See BindVisibilityToMultipleActivations */
+	void HandleVisibilityBoundWidgetActivations();
+
 	UPROPERTY(EditAnywhere, Category = Activation, meta = (InlineEditConditionToggle = "ActivatedVisibility"))
 	bool bSetVisibilityOnActivated = false;
 
@@ -177,6 +200,22 @@ private:
 	UPROPERTY(BlueprintReadOnly, Category = ActivatableWidget, meta = (AllowPrivateAccess = true))
 	bool bIsActive = false;
 
+	/** List of widgets whose collective activation controls our visibility. */
+	UPROPERTY(Transient)
+	TArray<TWeakObjectPtr<UCommonActivatableWidget>> VisibilityBoundWidgets;
+
+	/** Visibility to use when widgets we are bound to are activated */
+	ESlateVisibility ActivatedBindVisibility = ESlateVisibility::SelfHitTestInvisible;
+
+	/** Visibility to use when widgets we are bound to are deactivated, not used if widget has activation / deactivation visibilities */
+	ESlateVisibility DeactivatedBindVisibility = ESlateVisibility::SelfHitTestInvisible;
+
+	/** True if we should switch to activated visibility only when all bound widgets are active */
+	bool bAllActive = true;
+
+	/** Handle to default back action, if bound */
+	FUIActionBindingHandle DefaultBackActionHandle;
+
 	mutable FSimpleMulticastDelegate OnActivatedEvent;
 	mutable FSimpleMulticastDelegate OnDeactivatedEvent;
 	mutable FSimpleMulticastDelegate OnSlateReleasedEvent;
diff --git a/Engine/Plugins/Experimental/CommonUI/Source/CommonUI/Public/CommonUITypes.h b/Engine/Plugins/Experimental/CommonUI/Source/CommonUI/Public/CommonUITypes.h
--- a/Engine/Plugins/Experimental/CommonUI/Source/CommonUI/Public/CommonUITypes.h
+++ b/Engine/Plugins/Experimental/CommonUI/Source/CommonUI/Public/CommonUITypes.h
@@ -92,6 +92,10 @@ struct COMMONUI_API FCommonInputActionDataBase : public FTableRowBase
 	UPROPERTY(EditAnywhere, BlueprintReadOnly, Category = "CommonInput")
 	FText HoldDisplayName;
 
+	/** Priority in nav-bar */
+	UPROPERTY(EditAnywhere, BlueprintReadOnly, Category = "CommonInput")
+	int32 NavBarPriority = 0;
+
 protected:
 	/**
 	* Key to bind to for each input method
diff --git a/Engine/Plugins/Experimental/CommonUI/Source/CommonUI/Public/Input/CommonBoundActionBar.h b/Engine/Plugins/Experimental/CommonUI/Source/CommonUI/Public/Input/CommonBoundActionBar.h
--- a/Engine/Plugins/Experimental/CommonUI/Source/CommonUI/Public/Input/CommonBoundActionBar.h
+++ b/Engine/Plugins/Experimental/CommonUI/Source/CommonUI/Public/Input/CommonBoundActionBar.h
@@ -43,5 +43,8 @@ private:
 	UPROPERTY(EditAnywhere, Category = Display)
 	bool bDisplayOwningPlayerActionsOnly = true;
 
+	UPROPERTY(EditAnywhere, AdvancedDisplay, Category = Display)
+	bool bIgnoreDuplicateActions = true;
+
 	bool bIsRefreshQueued = false;
 };
\ No newline at end of file
