Put back in GetOptions restrictions for gamepads.

diff --git a/Engine/Plugins/Experimental/CommonUI/Source/CommonInput/Public/CommonInputBaseTypes.h b/Engine/Plugins/Experimental/CommonUI/Source/CommonInput/Public/CommonInputBaseTypes.h
--- a/Engine/Plugins/Experimental/CommonUI/Source/CommonInput/Public/CommonInputBaseTypes.h
+++ b/Engine/Plugins/Experimental/CommonUI/Source/CommonInput/Public/CommonInputBaseTypes.h
@@ -127,7 +127,7 @@ public:
 	UPROPERTY(EditDefaultsOnly, Category = "Default")
 	ECommonInputType InputType;
 	
-	UPROPERTY(EditDefaultsOnly, Category = "Gamepad", meta=(EditCondition="InputType == ECommonInputType::Gamepad"))
+	UPROPERTY(EditDefaultsOnly, Category = "Gamepad", meta=(EditCondition="InputType == ECommonInputType::Gamepad", GetOptions = GetRegisteredGamepads))
 	FName GamepadName;
 
 	UPROPERTY(EditDefaultsOnly, Category = "Gamepad", meta = (EditCondition = "InputType == ECommonInputType::Gamepad"))
