Lyra: Made a distinction between returning to the main menu and qutting PIE (only shown in the editor) in the in-game escape menu

diff --git a/Engine/Plugins/Experimental/CommonUI/Source/CommonUI/Private/CommonUISettings.cpp b/Engine/Plugins/Experimental/CommonUI/Source/CommonUI/Private/CommonUISettings.cpp
--- a/Engine/Plugins/Experimental/CommonUI/Source/CommonUI/Private/CommonUISettings.cpp
+++ b/Engine/Plugins/Experimental/CommonUI/Source/CommonUI/Private/CommonUISettings.cpp
@@ -9,6 +9,8 @@
 #include "GameplayTagContainer.h"
 #include "Misc/ConfigCacheIni.h"
 
+UE_DEFINE_GAMEPLAY_TAG(TAG_PlatformTrait_PlayInEditor, "Platform.Trait.PlayInEditor");
+
 UCommonUISettings::UCommonUISettings(const FObjectInitializer& Initializer)
 	: Super(Initializer)
 	, bAutoLoadData(true)
@@ -87,6 +89,13 @@ void UCommonUISettings::RebuildTraitContainer()
 	{
 		PlatformTraitContainer.AddTag(Trait);
 	}
+
+#if WITH_EDITOR
+	if (GIsEditor)
+	{
+		PlatformTraitContainer.AddTag(TAG_PlatformTrait_PlayInEditor);
+	}
+#endif
 }
 
 void UCommonUISettings::AutoLoadData()
diff --git a/Engine/Plugins/Experimental/CommonUI/Source/CommonUI/Public/CommonUISettings.h b/Engine/Plugins/Experimental/CommonUI/Source/CommonUI/Public/CommonUISettings.h
--- a/Engine/Plugins/Experimental/CommonUI/Source/CommonUI/Public/CommonUISettings.h
+++ b/Engine/Plugins/Experimental/CommonUI/Source/CommonUI/Public/CommonUISettings.h
@@ -12,10 +12,13 @@
 #include "CommonUIRichTextData.h"
 #include "CommonTextBlock.h"
 #include "GameplayTagContainer.h"
+#include "NativeGameplayTags.h"
 #include "CommonUISettings.generated.h"
 
 class UMaterial;
 
+COMMONUI_API UE_DECLARE_GAMEPLAY_TAG_EXTERN(TAG_PlatformTrait_PlayInEditor);
+
 UCLASS(config = Game, defaultconfig)
 class COMMONUI_API UCommonUISettings : public UObject
 {
