Fix Common button group disabling buttons unnecessarily

diff --git a/Engine/Plugins/Experimental/CommonUI/Source/CommonUI/Private/CommonButtonBase.cpp b/Engine/Plugins/Experimental/CommonUI/Source/CommonUI/Private/CommonButtonBase.cpp
--- a/Engine/Plugins/Experimental/CommonUI/Source/CommonUI/Private/CommonButtonBase.cpp
+++ b/Engine/Plugins/Experimental/CommonUI/Source/CommonUI/Private/CommonButtonBase.cpp
@@ -720,7 +720,7 @@ void UCommonButtonBase::SetSelectedInternal(bool bInSelected, bool bAllowSound /
 	if (bSelected)
 	{
 		NativeOnSelected(bBroadcast);
-		if (!bToggleable)
+		if (!bToggleable && IsInteractable())
 		{
 			// If the button isn't toggleable, then disable interaction with the root button while selected
 			// The prevents us getting unnecessary click noises and events
@@ -1375,4 +1375,4 @@ void UCommonButtonBase::SetIsFocusable(bool bInIsFocusable)
 bool UCommonButtonBase::GetIsFocusable() const
 {
 	return bIsFocusable;
-}
\ No newline at end of file
+}
