CommonUI - AutoLoad, needs to trigger in more spots or it wont work.

diff --git a/Engine/Plugins/Experimental/CommonUI/Source/CommonUI/Private/CommonUISettings.cpp b/Engine/Plugins/Experimental/CommonUI/Source/CommonUI/Private/CommonUISettings.cpp
--- a/Engine/Plugins/Experimental/CommonUI/Source/CommonUI/Private/CommonUISettings.cpp
+++ b/Engine/Plugins/Experimental/CommonUI/Source/CommonUI/Private/CommonUISettings.cpp
@@ -38,6 +38,7 @@ void UCommonUISettings::PostReloadConfig(FProperty* PropertyThatWasLoaded)
 	Super::PostReloadConfig(PropertyThatWasLoaded);
 
 	RebuildTraitContainer();
+	AutoLoadData();
 }
 
 void UCommonUISettings::PostInitProperties()
@@ -45,6 +46,7 @@ void UCommonUISettings::PostInitProperties()
 	Super::PostInitProperties();
 
 	RebuildTraitContainer();
+	AutoLoadData();
 }
 
 void UCommonUISettings::LoadEditorData()
