Don't crash when FindLocalPlayerFromControllerId returns null.

This can happen under various circumstances on console platforms, the calling code shouldn't assume that the return value will be valid.

#rnx

diff --git a/Engine/Plugins/Experimental/CommonUI/Source/CommonUI/Private/CommonGameViewportClient.cpp b/Engine/Plugins/Experimental/CommonUI/Source/CommonUI/Private/CommonGameViewportClient.cpp
--- a/Engine/Plugins/Experimental/CommonUI/Source/CommonUI/Private/CommonGameViewportClient.cpp
+++ b/Engine/Plugins/Experimental/CommonUI/Source/CommonUI/Private/CommonGameViewportClient.cpp
@@ -112,17 +112,20 @@ void UCommonGameViewportClient::HandleRerouteInput(int32 ControllerId, FKey Key,
 	ULocalPlayer* LocalPlayer = GameInstance->FindLocalPlayerFromControllerId(ControllerId);
 	Reply = FReply::Unhandled();
 
-	UCommonUIActionRouterBase* ActionRouter = LocalPlayer->GetSubsystem<UCommonUIActionRouterBase>();
-	if (ensure(ActionRouter))
+	if (LocalPlayer)
 	{
-		ERouteUIInputResult InputResult = ActionRouter->ProcessInput(Key, EventType);
-		if (InputResult == ERouteUIInputResult::BlockGameInput)
+		UCommonUIActionRouterBase* ActionRouter = LocalPlayer->GetSubsystem<UCommonUIActionRouterBase>();
+		if (ensure(ActionRouter))
 		{
-			OnRerouteBlockedInput().ExecuteIfBound(ControllerId, Key, EventType, Reply);
-		}
-		else if (InputResult == ERouteUIInputResult::Handled)
-		{
-			Reply = FReply::Handled();
+			ERouteUIInputResult InputResult = ActionRouter->ProcessInput(Key, EventType);
+			if (InputResult == ERouteUIInputResult::BlockGameInput)
+			{
+				OnRerouteBlockedInput().ExecuteIfBound(ControllerId, Key, EventType, Reply);
+			}
+			else if (InputResult == ERouteUIInputResult::Handled)
+			{
+				Reply = FReply::Handled();
+			}
 		}
 	}
 }
@@ -132,13 +135,16 @@ void UCommonGameViewportClient::HandleRerouteAxis(int32 ControllerId, FKey Key,
 	ULocalPlayer* LocalPlayer = GameInstance->FindLocalPlayerFromControllerId(ControllerId);
 	Reply = FReply::Unhandled();
 
-	UCommonUIActionRouterBase* ActionRouter = LocalPlayer->GetSubsystem<UCommonUIActionRouterBase>();
-	if (ensure(ActionRouter))
+	if (LocalPlayer)
 	{
-		// We don't actually use axis inputs that reach the game viewport UI land for anything, we just want block them reaching the game when they shouldn't
-		if (!ActionRouter->CanProcessNormalGameInput())
+		UCommonUIActionRouterBase* ActionRouter = LocalPlayer->GetSubsystem<UCommonUIActionRouterBase>();
+		if (ensure(ActionRouter))
 		{
-			Reply = FReply::Handled();
+			// We don't actually use axis inputs that reach the game viewport UI land for anything, we just want block them reaching the game when they shouldn't
+			if (!ActionRouter->CanProcessNormalGameInput())
+			{
+				Reply = FReply::Handled();
+			}
 		}
 	}
 }
@@ -148,7 +154,7 @@ void UCommonGameViewportClient::HandleRerouteTouch(int32 ControllerId, uint32 To
 	ULocalPlayer* LocalPlayer = GameInstance->FindLocalPlayerFromControllerId(ControllerId);
 	Reply = FReply::Unhandled();
 
-	if (TouchId < EKeys::NUM_TOUCH_KEYS)
+	if (LocalPlayer && TouchId < EKeys::NUM_TOUCH_KEYS)
 	{
 		FKey KeyPressed = EKeys::TouchKeys[TouchId];
 		if (KeyPressed.IsValid())
