Speculative fix CommonUI processing invalid widgets on tick.

diff --git a/Engine/Plugins/Experimental/CommonUI/Source/CommonUI/Private/Input/CommonUIActionRouterBase.cpp b/Engine/Plugins/Experimental/CommonUI/Source/CommonUI/Private/Input/CommonUIActionRouterBase.cpp
--- a/Engine/Plugins/Experimental/CommonUI/Source/CommonUI/Private/Input/CommonUIActionRouterBase.cpp
+++ b/Engine/Plugins/Experimental/CommonUI/Source/CommonUI/Private/Input/CommonUIActionRouterBase.cpp
@@ -575,6 +575,15 @@ void UCommonUIActionRouterBase::HandleSlateFocusChanging(const FFocusEvent& Focu
 void UCommonUIActionRouterBase::HandlePostGarbageCollect()
 {
 	FUIActionBinding::CleanRegistrations();
+
+	// GC may result in root widget being purged while conditional slate resource release skips HandleRootWidgetSlateReleased, handle this scenario
+	for (auto Iter = RootNodes.CreateIterator(); Iter; ++Iter)
+	{
+		if (!Iter->Get().IsWidgetValid())
+		{
+			Iter.RemoveCurrent();
+		}
+	}
 }
 
 void UCommonUIActionRouterBase::ProcessRebuiltWidgets()
diff --git a/Engine/Plugins/Experimental/CommonUI/Source/CommonUI/Private/Input/UIActionRouterTypes.cpp b/Engine/Plugins/Experimental/CommonUI/Source/CommonUI/Private/Input/UIActionRouterTypes.cpp
--- a/Engine/Plugins/Experimental/CommonUI/Source/CommonUI/Private/Input/UIActionRouterTypes.cpp
+++ b/Engine/Plugins/Experimental/CommonUI/Source/CommonUI/Private/Input/UIActionRouterTypes.cpp
@@ -715,6 +715,11 @@ bool FActivatableTreeNode::ProcessNormalInput(ECommonInputMode ActiveInputMode,
 	return false;
 }
 
+bool FActivatableTreeNode::IsWidgetValid() const
+{ 
+	return RepresentedWidget.IsValid();
+}
+
 bool FActivatableTreeNode::IsWidgetActivated() const
 {
 	return ensure(RepresentedWidget.IsValid()) && RepresentedWidget->IsActivated();
diff --git a/Engine/Plugins/Experimental/CommonUI/Source/CommonUI/Private/Input/UIActionRouterTypes.h b/Engine/Plugins/Experimental/CommonUI/Source/CommonUI/Private/Input/UIActionRouterTypes.h
--- a/Engine/Plugins/Experimental/CommonUI/Source/CommonUI/Private/Input/UIActionRouterTypes.h
+++ b/Engine/Plugins/Experimental/CommonUI/Source/CommonUI/Private/Input/UIActionRouterTypes.h
@@ -162,7 +162,8 @@ public:
 	virtual EProcessHoldActionResult ProcessHoldInput(ECommonInputMode ActiveInputMode, FKey Key, EInputEvent InputEvent) const override;
 	virtual bool ProcessNormalInput(ECommonInputMode ActiveInputMode, FKey Key, EInputEvent InputEvent) const override;
 	virtual bool IsReceivingInput() const override { return bCanReceiveInput && IsWidgetActivated(); }
-	
+
+	bool IsWidgetValid() const;
 	bool IsWidgetActivated() const;
 	bool DoesWidgetSupportActivationFocus() const;
 	void AppendAllActiveActions(TArray<FUIActionBindingHandle>& BoundActions) const;
