Fix sound hover regression.

diff --git a/Engine/Plugins/Experimental/CommonUI/Source/CommonUI/Private/CommonButtonTypes.cpp b/Engine/Plugins/Experimental/CommonUI/Source/CommonUI/Private/CommonButtonTypes.cpp
--- a/Engine/Plugins/Experimental/CommonUI/Source/CommonUI/Private/CommonButtonTypes.cpp
+++ b/Engine/Plugins/Experimental/CommonUI/Source/CommonUI/Private/CommonButtonTypes.cpp
@@ -67,9 +67,17 @@ void SCommonButton::OnMouseEnter(const FGeometry& MyGeometry, const FPointerEven
 {
 	if (!InPointerEvent.IsTouchEvent())
 	{
+		const bool bWasHovered = IsHovered();
+
 		bHovered = true;
 		SetHover(bHovered && bIsInteractionEnabled);
 		SButton::OnMouseEnter(MyGeometry, InPointerEvent);
+
+		// SButton won't be able to correctly detect hover changes since we manually set hover, do our own detection
+		if (!bWasHovered && IsHovered())
+		{
+			ExecuteHoverStateChanged(true);
+		}
 	}
 }
 
@@ -84,9 +92,17 @@ void SCommonButton::OnMouseLeave(const FPointerEvent& InPointerEvent)
 	}
 	else
 	{
+		const bool bWasHovered = IsHovered();
+
 		bHovered = false;
 		SetHover(false);
 		SButton::OnMouseLeave(InPointerEvent);
+
+		// SButton won't be able to correctly detect hover changes since we manually set hover, do our own detection
+		if (bWasHovered && !IsHovered())
+		{
+			ExecuteHoverStateChanged(true);
+		}
 	}
 }
 
