Made UCommonTabListWidgetBase::GetTabButtonBaseByID public.

diff --git a/Engine/Plugins/Experimental/CommonUI/Source/CommonUI/Public/CommonTabListWidgetBase.h b/Engine/Plugins/Experimental/CommonUI/Source/CommonUI/Public/CommonTabListWidgetBase.h
--- a/Engine/Plugins/Experimental/CommonUI/Source/CommonUI/Public/CommonTabListWidgetBase.h
+++ b/Engine/Plugins/Experimental/CommonUI/Source/CommonUI/Public/CommonTabListWidgetBase.h
@@ -130,6 +130,10 @@ public:
 	UFUNCTION(BlueprintCallable, Category = TabList)
 	virtual void SetListeningForInput(bool bShouldListen);
 
+	/** Returns the tab button matching the ID, if found */
+	UFUNCTION(BlueprintCallable, Category = TabList)
+	UCommonButtonBase* GetTabButtonBaseByID(FName TabNameID);
+
 protected:
 	// UUserWidget interface
 	virtual void NativeOnInitialized() override;
@@ -155,10 +159,6 @@ protected:
 	UFUNCTION(BlueprintNativeEvent, Category = TabList, meta = (BlueprintProtected = "true"))
 	void HandleTabRemoval(FName TabNameID, UCommonButtonBase* TabButton);
 
-	/** Returns the tab button matching the ID, if found */
-	UFUNCTION(BlueprintCallable, Category = TabList, meta = (BlueprintProtected = "true"))
-	UCommonButtonBase* GetTabButtonBaseByID(FName TabNameID);
-
 	/** The input action to listen for causing the next tab to be selected */
 	UPROPERTY(EditAnywhere, BlueprintReadOnly, Category = TabList, meta = (RowType = CommonInputActionDataBase))
 	FDataTableRowHandle NextTabInputActionData;
