Adding some logging to try and track down a rare issue where the map has no input

diff --git a/Engine/Plugins/Experimental/CommonUI/Source/CommonUI/Private/Input/CommonUIActionRouterBase.cpp b/Engine/Plugins/Experimental/CommonUI/Source/CommonUI/Private/Input/CommonUIActionRouterBase.cpp
--- a/Engine/Plugins/Experimental/CommonUI/Source/CommonUI/Private/Input/CommonUIActionRouterBase.cpp
+++ b/Engine/Plugins/Experimental/CommonUI/Source/CommonUI/Private/Input/CommonUIActionRouterBase.cpp
@@ -983,6 +983,9 @@ void UCommonUIActionRouterBase::ApplyUIInputConfig(const FUIInputConfig& NewConf
 {
 	if (bForceRefresh || NewConfig != ActiveInputConfig.GetValue())
 	{
+		UE_LOG(LogUIActionRouter, Display, TEXT("UIInputConfig being changed. bForceRefresh: %d"), bForceRefresh ? 1 : 0);
+		UE_LOG(LogUIActionRouter, Display, TEXT("\tInputMode: Previous (%d), New (%d)"), ActiveInputConfig.IsSet() ? (int32)ActiveInputConfig->GetInputMode() : -1, (int32)NewConfig.GetInputMode());
+
 		const ECommonInputMode PreviousInputMode = GetActiveInputMode();
 
 		ActiveInputConfig = NewConfig;
@@ -1046,8 +1049,20 @@ void UCommonUIActionRouterBase::ApplyUIInputConfig(const FUIInputConfig& NewConf
 					break;
 					}
 				}
+				else
+				{
+					UE_LOG(LogUIActionRouter, Warning, TEXT("\tFailed to commit change! Local player controller was null."));
+				}
+			}
+			else
+			{
+				UE_LOG(LogUIActionRouter, Warning, TEXT("\tFailed to commit change! ViewportWidget was null."));
 			}
 		}
+		else
+		{
+			UE_LOG(LogUIActionRouter, Warning, TEXT("\tFailed to commit change! GameViewportClient was null."));
+		}
 
 		if (PreviousInputMode != NewConfig.GetInputMode())
 		{
