Fix CommonUI processing when it shouldn't.

diff --git a/Engine/Plugins/Experimental/CommonUI/Source/CommonUI/Private/Input/CommonUIActionRouterBase.cpp b/Engine/Plugins/Experimental/CommonUI/Source/CommonUI/Private/Input/CommonUIActionRouterBase.cpp
--- a/Engine/Plugins/Experimental/CommonUI/Source/CommonUI/Private/Input/CommonUIActionRouterBase.cpp
+++ b/Engine/Plugins/Experimental/CommonUI/Source/CommonUI/Private/Input/CommonUIActionRouterBase.cpp
@@ -14,6 +14,7 @@
 #include "Engine/Canvas.h"
 #include "Stats/Stats.h"
 
+#include "CommonGameViewportClient.h"
 #include "CommonUIPrivatePCH.h"
 #include "CommonActivatableWidget.h"
 #include "CommonUIUtils.h"
@@ -203,7 +204,11 @@ void UCommonUIActionRouterBase::PostAnalogCursorCreate()
 
 void UCommonUIActionRouterBase::RegisterAnalogCursorTick()
 {
-	FSlateApplication::Get().RegisterInputPreProcessor(AnalogCursor, UCommonUIInputSettings::Get().GetAnalogCursorSettings().PreprocessorPriority);
+	if (GEngine->GameViewportClientClass->IsChildOf<UCommonGameViewportClient>())
+	{
+		FSlateApplication::Get().RegisterInputPreProcessor(AnalogCursor, UCommonUIInputSettings::Get().GetAnalogCursorSettings().PreprocessorPriority);
+	}
+
 	if (bIsActivatableTreeEnabled)
 	{
 		FTSTicker::GetCoreTicker().RemoveTicker(TickHandle);
