Implementing Common Input Setting to allow disabling the setting of a default Input Config when no Activatable Widgets request one.

diff --git a/Engine/Plugins/Experimental/CommonUI/Source/CommonInput/Public/CommonInputSettings.h b/Engine/Plugins/Experimental/CommonUI/Source/CommonInput/Public/CommonInputSettings.h
--- a/Engine/Plugins/Experimental/CommonUI/Source/CommonInput/Public/CommonInputSettings.h
+++ b/Engine/Plugins/Experimental/CommonUI/Source/CommonInput/Public/CommonInputSettings.h
@@ -47,6 +47,8 @@ public:
 
 	bool GetAllowOutOfFocusDeviceInput() const { return bAllowOutOfFocusDeviceInput; }
 
+	bool GetEnableDefaultInputConfig() const { return bEnableDefaultInputConfig; }
+
 private:
 	virtual void PostInitProperties() override;
 
@@ -75,6 +77,13 @@ private:
 	UPROPERTY(config, EditAnywhere, Category = "Input")
 	bool bAllowOutOfFocusDeviceInput = false;
 
+	/**
+	* Controls whether a default Input Config will be set when the active CommonActivatableWidgets do not specify a desired one.
+	* Disable this if you want to control the Input Mode via alternative means.
+	*/
+	UPROPERTY(config, EditAnywhere, Category = "Input")
+	bool bEnableDefaultInputConfig = true;
+
 private:
 	void LoadInputData();
 
diff --git a/Engine/Plugins/Experimental/CommonUI/Source/CommonUI/Private/Input/CommonUIActionRouterBase.cpp b/Engine/Plugins/Experimental/CommonUI/Source/CommonUI/Private/Input/CommonUIActionRouterBase.cpp
--- a/Engine/Plugins/Experimental/CommonUI/Source/CommonUI/Private/Input/CommonUIActionRouterBase.cpp
+++ b/Engine/Plugins/Experimental/CommonUI/Source/CommonUI/Private/Input/CommonUIActionRouterBase.cpp
@@ -4,6 +4,8 @@
 #include "Input/UIActionRouterTypes.h"
 #include "Input/CommonAnalogCursor.h"
 #include "Input/CommonUIInputSettings.h"
+#include "CommonInputSettings.h"
+#include "ICommonInputModule.h"
 
 #include "Framework/Application/SlateUser.h"
 #include "Slate/SObjectWidget.h"
@@ -1011,6 +1013,7 @@ FActivatableTreeNodePtr UCommonUIActionRouterBase::FindNodeRecursive(const FActi
 
 void UCommonUIActionRouterBase::SetActiveUIInputConfig(const FUIInputConfig& NewConfig)
 {
+	ensureMsgf(ICommonInputModule::GetSettings().GetEnableDefaultInputConfig() == false, TEXT("UIInputConfig is being set but will be overwritten when the activatable tree changes. Disable the default Input Config in Common Input Settings to prevent this."));
 	ApplyUIInputConfig(NewConfig, !ActiveInputConfig.IsSet());
 }
 
diff --git a/Engine/Plugins/Experimental/CommonUI/Source/CommonUI/Private/Input/UIActionRouterTypes.cpp b/Engine/Plugins/Experimental/CommonUI/Source/CommonUI/Private/Input/UIActionRouterTypes.cpp
--- a/Engine/Plugins/Experimental/CommonUI/Source/CommonUI/Private/Input/UIActionRouterTypes.cpp
+++ b/Engine/Plugins/Experimental/CommonUI/Source/CommonUI/Private/Input/UIActionRouterTypes.cpp
@@ -5,7 +5,9 @@
 #include "CommonActivatableWidget.h"
 #include "Input/CommonUIInputTypes.h"
 #include "CommonUIPrivatePCH.h"
+#include "CommonInputSettings.h"
 #include "CommonInputSubsystem.h"
+#include "ICommonInputModule.h"
 #include "Input/CommonUIInputSettings.h"
 #include "Engine/GameViewportClient.h"
 #include "Framework/Application/SlateApplication.h"
@@ -783,7 +785,7 @@ int32 FActivatableTreeNode::GetLastPaintLayer() const
 	return INDEX_NONE;
 }
 
-FUIInputConfig FActivatableTreeNode::FindDesiredInputConfig() const
+TOptional<FUIInputConfig> FActivatableTreeNode::FindDesiredInputConfig() const
 {
 	TOptional<FUIInputConfig> DesiredConfig = ensure(RepresentedWidget.IsValid()) ? RepresentedWidget->GetDesiredInputConfig() : TOptional<FUIInputConfig>();
 	if (!DesiredConfig.IsSet() && Parent.IsValid())
@@ -791,13 +793,7 @@ FUIInputConfig FActivatableTreeNode::FindDesiredInputConfig() const
 		DesiredConfig = Parent.Pin()->FindDesiredInputConfig();
 	}
 
-	if (DesiredConfig.IsSet())
-	{
-		return DesiredConfig.GetValue();
-	}
-
-	// If nobody in the entire tree cares about the config, fall back to the initial default
-	return FUIInputConfig();
+	return DesiredConfig;
 }
 
 FUICameraConfig FActivatableTreeNode::FindDesiredCameraConfig() const
@@ -1229,8 +1225,16 @@ void FActivatableTreeRoot::ApplyLeafmostNodeConfig()
 		{
 			UE_LOG(LogUIActionRouter, Display, TEXT("Applying input config for leaf-most node [%s]"), *PinnedLeafmostNode->GetWidget()->GetName());
 
-			FUIInputConfig DesiredConfig = PinnedLeafmostNode->FindDesiredInputConfig();
-			GetActionRouter().SetActiveUIInputConfig(DesiredConfig);
+			TOptional<FUIInputConfig> DesiredConfig = PinnedLeafmostNode->FindDesiredInputConfig();
+			if(DesiredConfig.IsSet())
+			{
+				GetActionRouter().SetActiveUIInputConfig(DesiredConfig.GetValue());
+			}
+			else if(ICommonInputModule::GetSettings().GetEnableDefaultInputConfig())
+			{
+				// Nobody in the entire tree cares about the config and the default is enabled so fall back to the default
+				GetActionRouter().SetActiveUIInputConfig(FUIInputConfig());
+			}
 
 			FocusLeafmostNode();
 		}
diff --git a/Engine/Plugins/Experimental/CommonUI/Source/CommonUI/Private/Input/UIActionRouterTypes.h b/Engine/Plugins/Experimental/CommonUI/Source/CommonUI/Private/Input/UIActionRouterTypes.h
--- a/Engine/Plugins/Experimental/CommonUI/Source/CommonUI/Private/Input/UIActionRouterTypes.h
+++ b/Engine/Plugins/Experimental/CommonUI/Source/CommonUI/Private/Input/UIActionRouterTypes.h
@@ -187,7 +187,7 @@ public:
 	bool IsExclusiveParentOfWidget(const TSharedPtr<SWidget>& SlateWidget) const;
 
 	int32 GetLastPaintLayer() const;
-	FUIInputConfig FindDesiredInputConfig() const;
+	TOptional<FUIInputConfig> FindDesiredInputConfig() const;
 	FUICameraConfig FindDesiredCameraConfig() const;
 
 	void AddScrollRecipient(const UWidget& ScrollRecipient);
diff --git a/Engine/Plugins/Experimental/CommonUI/Source/CommonUI/Public/Input/CommonInputMode.h b/Engine/Plugins/Experimental/CommonUI/Source/CommonUI/Public/Input/CommonInputMode.h
--- a/Engine/Plugins/Experimental/CommonUI/Source/CommonUI/Public/Input/CommonInputMode.h
+++ b/Engine/Plugins/Experimental/CommonUI/Source/CommonUI/Public/Input/CommonInputMode.h
@@ -9,8 +9,9 @@
 UENUM(BlueprintType)
 enum class ECommonInputMode : uint8
 {
-	Menu,
-	Game,
-	All,
-	MAX
+	Menu	UMETA(Tooltip="Input is received by the UI only"),
+	Game	UMETA(Tooltip="Input is received by the Game only"),
+	All		UMETA(Tooltip="Input is received by UI and the Game"),
+	
+	MAX UMETA(Hidden)
 };
