CommonUI - Speeding up simulated trait queries so that they're not recalculated every time the traits are requested during simulation.

diff --git a/Engine/Plugins/Experimental/CommonUI/Source/CommonUI/Private/CommonUISettings.cpp b/Engine/Plugins/Experimental/CommonUI/Source/CommonUI/Private/CommonUISettings.cpp
--- a/Engine/Plugins/Experimental/CommonUI/Source/CommonUI/Private/CommonUISettings.cpp
+++ b/Engine/Plugins/Experimental/CommonUI/Source/CommonUI/Private/CommonUISettings.cpp
@@ -135,28 +135,36 @@ const FGameplayTagContainer& UCommonUISettings::GetPlatformTraits() const
 	const FName SimulatedPlatform = UPlatformSettings::GetEditorSimulatedPlatform();
 	if (SimulatedPlatform != NAME_None)
 	{
-		if (FConfigCacheIni* PlatformIni = FConfigCacheIni::ForPlatform(SimulatedPlatform))
-		{
-			TArray<FString> Entries;
-			PlatformIni->GetArray(*GetClass()->GetPathName(), GET_MEMBER_NAME_STRING_CHECKED(UCommonUISettings, PlatformTraits), Entries, TEXT("Game"));
+		static FName LastSimulatedPlatform = NAME_None;
+		static FGameplayTagContainer SimulatedPlatformTraitContainer;
 
-			UScriptStruct* GameplayTagStruct = StaticStruct<FGameplayTag>();
-			TArray<FGameplayTag> OtherPlatformTraits;
-			for(FString Entry : Entries)
-			{
-				FGameplayTag Element;
-				GameplayTagStruct->ImportText(*Entry, &Element, nullptr, EPropertyPortFlags::PPF_None, nullptr, GameplayTagStruct->GetName(), true);
-				OtherPlatformTraits.Add(Element);
-			}
+		if (SimulatedPlatform != LastSimulatedPlatform)
+		{
+			LastSimulatedPlatform = SimulatedPlatform;
+			SimulatedPlatformTraitContainer.Reset();
 
-			static FGameplayTagContainer OtherPlatformTraitContainer;
-			OtherPlatformTraitContainer.Reset();
-			for (FGameplayTag Trait : OtherPlatformTraits)
+			if (FConfigCacheIni* PlatformIni = FConfigCacheIni::ForPlatform(SimulatedPlatform))
 			{
-				OtherPlatformTraitContainer.AddTag(Trait);
+				TArray<FString> Entries;
+				PlatformIni->GetArray(*GetClass()->GetPathName(), GET_MEMBER_NAME_STRING_CHECKED(UCommonUISettings, PlatformTraits), Entries, TEXT("Game"));
+
+				UScriptStruct* GameplayTagStruct = StaticStruct<FGameplayTag>();
+				TArray<FGameplayTag> OtherPlatformTraits;
+				for(FString Entry : Entries)
+				{
+					FGameplayTag Element;
+					GameplayTagStruct->ImportText(*Entry, &Element, nullptr, EPropertyPortFlags::PPF_None, nullptr, GameplayTagStruct->GetName(), true);
+					OtherPlatformTraits.Add(Element);
+				}
+
+				for (FGameplayTag Trait : OtherPlatformTraits)
+				{
+					SimulatedPlatformTraitContainer.AddTag(Trait);
+				}
 			}
-			return OtherPlatformTraitContainer;
 		}
+
+		return SimulatedPlatformTraitContainer;
 	}
 #endif
 
