Fix issue with UCommonRichText where the default text style override wasn't directly applied if Mobile Scale is enabled
#RNX

diff --git a/Engine/Plugins/Experimental/CommonUI/Source/CommonUI/Private/CommonRichTextBlock.cpp b/Engine/Plugins/Experimental/CommonUI/Source/CommonUI/Private/CommonRichTextBlock.cpp
--- a/Engine/Plugins/Experimental/CommonUI/Source/CommonUI/Private/CommonRichTextBlock.cpp
+++ b/Engine/Plugins/Experimental/CommonUI/Source/CommonUI/Private/CommonRichTextBlock.cpp
@@ -594,14 +594,12 @@ void UCommonRichTextBlock::SynchronizeProperties()
 
 void UCommonRichTextBlock::ApplyUpdatedDefaultTextStyle()
 {
+	Super::ApplyUpdatedDefaultTextStyle();
+
 	if (CommonUIUtils::ShouldDisplayMobileUISizes())
 	{
 		ApplyTextBlockScale();
 	}
-	else
-	{
-		Super::ApplyUpdatedDefaultTextStyle();
-	}
 }
 
 void UCommonRichTextBlock::CreateDecorators(TArray<TSharedRef<ITextDecorator>>& OutDecorators)
