Ensure local player exists when changing root node.

diff --git a/Engine/Plugins/Experimental/CommonUI/Source/CommonUI/Private/Input/CommonUIActionRouterBase.cpp b/Engine/Plugins/Experimental/CommonUI/Source/CommonUI/Private/Input/CommonUIActionRouterBase.cpp
--- a/Engine/Plugins/Experimental/CommonUI/Source/CommonUI/Private/Input/CommonUIActionRouterBase.cpp
+++ b/Engine/Plugins/Experimental/CommonUI/Source/CommonUI/Private/Input/CommonUIActionRouterBase.cpp
@@ -557,7 +557,12 @@ void UCommonUIActionRouterBase::HandleRootNodeActivated(TWeakPtr<FActivatableTre
 		const int32 CurrentRootLayer = ActiveRootNode ? ActiveRootNode->GetLastPaintLayer() : INDEX_NONE;
 		if (ActivatedRoot->GetLastPaintLayer() > CurrentRootLayer)
 		{
-			SetActiveRoot(ActivatedRoot);
+			// Ensure we have a local player so the action router local player subsystem will handle root change
+			const ULocalPlayer* LocalPlayer = GetLocalPlayer();
+			if (LocalPlayer && LocalPlayer->ViewportClient)
+			{
+				SetActiveRoot(ActivatedRoot);
+			}
 		}
 	}
 }
