Fix platform definition in CommonUI

#rnx
#p4v-preflight-copy 18430652

diff --git a/Engine/Plugins/Experimental/CommonUI/Source/CommonUI/Private/CommonUITypes.cpp b/Engine/Plugins/Experimental/CommonUI/Source/CommonUI/Private/CommonUITypes.cpp
--- a/Engine/Plugins/Experimental/CommonUI/Source/CommonUI/Private/CommonUITypes.cpp
+++ b/Engine/Plugins/Experimental/CommonUI/Source/CommonUI/Private/CommonUITypes.cpp
@@ -4,6 +4,7 @@
 #include "CommonUIPrivatePCH.h"
 #include "ICommonInputModule.h"
 #include "CommonInputSettings.h"
+#include "HAL/PlatformInput.h"
 
 FScrollBoxStyle CommonUI::EmptyScrollBoxStyle = FScrollBoxStyle();
 
@@ -17,21 +18,7 @@ FCommonInputTypeInfo::FCommonInputTypeInfo()
 
 FKey FCommonInputTypeInfo::GetKey() const
 {
-	FKey ReturnKey = Key;
-#if PLATFORM_PS4
-	if (EKeys::Virtual_Accept == EKeys::Gamepad_FaceButton_Right)
-	{
-		if (ReturnKey == EKeys::Gamepad_FaceButton_Bottom)
-		{
-			ReturnKey = EKeys::Gamepad_FaceButton_Right;
-		}
-		else if (ReturnKey == EKeys::Gamepad_FaceButton_Right)
-		{
-			ReturnKey = EKeys::Gamepad_FaceButton_Bottom;
-		}
-	}
-#endif
-	return ReturnKey;
+	return FPlatformInput::RemapKey(Key);
 }
 
 FCommonInputActionDataBase::FCommonInputActionDataBase()
diff --git a/Engine/Source/Runtime/InputCore/Public/GenericPlatform/GenericPlatformInput.h b/Engine/Source/Runtime/InputCore/Public/GenericPlatform/GenericPlatformInput.h
--- a/Engine/Source/Runtime/InputCore/Public/GenericPlatform/GenericPlatformInput.h
+++ b/Engine/Source/Runtime/InputCore/Public/GenericPlatform/GenericPlatformInput.h
@@ -18,6 +18,9 @@ public:
 		return 0;
 	}
 
+	/** Remap a given key to something else if necessary. */
+	FORCEINLINE static FKey RemapKey(FKey Key) { return Key; }
+
 	static FKey GetGamepadAcceptKey()
 	{
 		return EKeys::Gamepad_FaceButton_Bottom;
