Fixing an issue with the CommonActionWidget and global invalidation. The brush being passed into the SetImage call is the same object every time, despite the brush changing internally. Because of this no invalidation is getting triggered, so we must manually do it.

diff --git a/Engine/Plugins/Experimental/CommonUI/Source/CommonUI/Private/CommonActionWidget.cpp b/Engine/Plugins/Experimental/CommonUI/Source/CommonUI/Private/CommonActionWidget.cpp
--- a/Engine/Plugins/Experimental/CommonUI/Source/CommonUI/Private/CommonActionWidget.cpp
+++ b/Engine/Plugins/Experimental/CommonUI/Source/CommonUI/Private/CommonActionWidget.cpp
@@ -242,6 +242,11 @@ void UCommonActionWidget::UpdateActionWidget()
 					{
 						MyIcon->SetImage(&Icon);
 
+						if (GetVisibility() != ESlateVisibility::Collapsed)
+						{
+							MyIcon->Invalidate(EInvalidateWidgetReason::LayoutAndVolatility);
+						}
+
 						if (InputActionData->GetCurrentInputTypeInfo(CommonInputSubsystem).bActionRequiresHold)
 						{
 							MyProgressImage->SetVisibility(EVisibility::SelfHitTestInvisible);
