Addressing some feedback from review 18456222

diff --git a/Engine/Plugins/Experimental/CommonUI/Source/CommonUI/Private/CommonActionWidget.cpp b/Engine/Plugins/Experimental/CommonUI/Source/CommonUI/Private/CommonActionWidget.cpp
--- a/Engine/Plugins/Experimental/CommonUI/Source/CommonUI/Private/CommonActionWidget.cpp
+++ b/Engine/Plugins/Experimental/CommonUI/Source/CommonUI/Private/CommonActionWidget.cpp
@@ -244,7 +244,9 @@ void UCommonActionWidget::UpdateActionWidget()
 
 						if (GetVisibility() != ESlateVisibility::Collapsed)
 						{
-							MyIcon->Invalidate(EInvalidateWidgetReason::LayoutAndVolatility);
+							// The object being passed into SetImage is the same each time so layout is never invalidated
+							// Manually invalidate it here as the dimensions may have changed
+							MyIcon->Invalidate(EInvalidateWidgetReason::Layout);
 						}
 
 						if (InputActionData->GetCurrentInputTypeInfo(CommonInputSubsystem).bActionRequiresHold)
