Change log category of CommonUI ActionRouter derived class search.

diff --git a/Engine/Plugins/Experimental/CommonUI/Source/CommonUI/Private/Input/CommonUIActionRouterBase.cpp b/Engine/Plugins/Experimental/CommonUI/Source/CommonUI/Private/Input/CommonUIActionRouterBase.cpp
--- a/Engine/Plugins/Experimental/CommonUI/Source/CommonUI/Private/Input/CommonUIActionRouterBase.cpp
+++ b/Engine/Plugins/Experimental/CommonUI/Source/CommonUI/Private/Input/CommonUIActionRouterBase.cpp
@@ -241,7 +241,7 @@ bool UCommonUIActionRouterBase::ShouldCreateSubsystem(UObject* Outer) const
 	TArray<UClass*> ChildClasses;
 	GetDerivedClasses(GetClass(), ChildClasses, false);
 
-	UE_LOG(LogUIActionRouter, Warning, TEXT("Found %i derived classes when attemping to create action router (%s)"), ChildClasses.Num(), *GetClass()->GetName());
+	UE_LOG(LogUIActionRouter, Display, TEXT("Found %i derived classes when attemping to create action router (%s)"), ChildClasses.Num(), *GetClass()->GetName());
 
 	// Only create an instance if there is no override implementation defined elsewhere
 	return ChildClasses.Num() == 0;
