Guard against action router not existing on game shutdown.

diff --git a/Engine/Plugins/Experimental/CommonUI/Source/CommonUI/Private/CommonUserWidget.cpp b/Engine/Plugins/Experimental/CommonUI/Source/CommonUI/Private/CommonUserWidget.cpp
--- a/Engine/Plugins/Experimental/CommonUI/Source/CommonUI/Private/CommonUserWidget.cpp
+++ b/Engine/Plugins/Experimental/CommonUI/Source/CommonUI/Private/CommonUserWidget.cpp
@@ -98,13 +98,19 @@ FUIActionBindingHandle UCommonUserWidget::RegisterUIActionBinding(const FBindUIA
 void UCommonUserWidget::RemoveActionBinding(FUIActionBindingHandle ActionBinding)
 {
 	ActionBindings.Remove(ActionBinding);
-	UCommonUIActionRouterBase::Get(*this)->RemoveBinding(ActionBinding);
+	if (UCommonUIActionRouterBase* ActionRouter = UCommonUIActionRouterBase::Get(*this))
+	{
+		ActionRouter->RemoveBinding(ActionBinding);
+	}
 }
 
 void UCommonUserWidget::AddActionBinding(FUIActionBindingHandle ActionBinding)
 {
 	ActionBindings.Add(ActionBinding);
-	UCommonUIActionRouterBase::Get(*this)->AddBinding(ActionBinding);
+	if (UCommonUIActionRouterBase* ActionRouter = UCommonUIActionRouterBase::Get(*this))
+	{
+		ActionRouter->AddBinding(ActionBinding);
+	}
 }
 
 void UCommonUserWidget::RegisterScrollRecipient(const UWidget& AnalogScrollRecipient)
