Added bCurrentlySwitching variable and const getter to UCommonAnimatedSwitcher.

diff --git a/Engine/Plugins/Experimental/CommonUI/Source/CommonUI/Private/CommonAnimatedSwitcher.cpp b/Engine/Plugins/Experimental/CommonUI/Source/CommonUI/Private/CommonAnimatedSwitcher.cpp
--- a/Engine/Plugins/Experimental/CommonUI/Source/CommonUI/Private/CommonAnimatedSwitcher.cpp
+++ b/Engine/Plugins/Experimental/CommonUI/Source/CommonUI/Private/CommonAnimatedSwitcher.cpp
@@ -1,10 +1,13 @@
 // Copyright Epic Games, Inc. All Rights Reserved.
+
 #include "CommonAnimatedSwitcher.h"
+
 #include "CommonWidgetPaletteCategories.h"
-#include "Components/WidgetSwitcherSlot.h"
 #include "Components/PanelSlot.h"
-#include "Widgets/SOverlay.h"
+#include "Components/WidgetSwitcherSlot.h"
+#include "Templates/UnrealTemplate.h"
 #include "Widgets/Layout/SSpacer.h"
+#include "Widgets/SOverlay.h"
 
 UCommonAnimatedSwitcher::UCommonAnimatedSwitcher(const FObjectInitializer& ObjectInitializer)
 	: Super(ObjectInitializer)
@@ -90,6 +93,11 @@ void UCommonAnimatedSwitcher::SetDisableTransitionAnimation(bool bDisableAnimati
 	bInstantTransition = bDisableAnimation;
 }
 
+bool UCommonAnimatedSwitcher::IsCurrentlySwitching() const
+{
+	return bCurrentlySwitching;
+}
+
 void UCommonAnimatedSwitcher::HandleSlateActiveIndexChanged(int32 ActiveIndex)
 {
 	if (Slots.IsValidIndex(ActiveWidgetIndex))
@@ -144,6 +152,8 @@ void UCommonAnimatedSwitcher::SetActiveWidgetIndex_Internal(int32 Index)
 		return;
 	}
 #endif
+	
+	TGuardValue<bool> bCurrentlySwitchingGuard(bCurrentlySwitching, true);
 
 	if (Index >= 0 && Index < Slots.Num() && (Index != ActiveWidgetIndex || !bSetOnce))
 	{
diff --git a/Engine/Plugins/Experimental/CommonUI/Source/CommonUI/Public/CommonAnimatedSwitcher.h b/Engine/Plugins/Experimental/CommonUI/Source/CommonUI/Public/CommonAnimatedSwitcher.h
--- a/Engine/Plugins/Experimental/CommonUI/Source/CommonUI/Public/CommonAnimatedSwitcher.h
+++ b/Engine/Plugins/Experimental/CommonUI/Source/CommonUI/Public/CommonAnimatedSwitcher.h
@@ -43,6 +43,9 @@ public:
 	UFUNCTION(BlueprintCallable, Category = "Common Widget Switcher")
 	void SetDisableTransitionAnimation(bool bDisableAnimation);
 
+	UFUNCTION(BlueprintCallable, Category = "Common Widget Switcher")
+	bool IsCurrentlySwitching() const;
+
 protected:
 	virtual void HandleSlateActiveIndexChanged(int32 ActiveIndex);
 
@@ -78,6 +81,7 @@ protected:
 	bool bInstantTransition = false;
 
 	bool bSetOnce = false;
+	bool bCurrentlySwitching = false;
 
 private:
 	void SetActiveWidgetIndex_Internal(int32 Index);
