Fix several shutdown crashes

diff --git a/Engine/Plugins/Experimental/CommonUI/Source/CommonInput/Private/CommonInputSubsystem.cpp b/Engine/Plugins/Experimental/CommonUI/Source/CommonInput/Private/CommonInputSubsystem.cpp
--- a/Engine/Plugins/Experimental/CommonUI/Source/CommonInput/Private/CommonInputSubsystem.cpp
+++ b/Engine/Plugins/Experimental/CommonUI/Source/CommonInput/Private/CommonInputSubsystem.cpp
@@ -328,7 +328,10 @@ void UCommonInputSubsystem::Initialize(FSubsystemCollectionBase& Collection)
 void UCommonInputSubsystem::Deinitialize()
 {
 	Super::Deinitialize();
-	FSlateApplication::Get().UnregisterInputPreProcessor(CommonInputPreprocessor);
+	if (FSlateApplication::IsInitialized())
+	{
+		FSlateApplication::Get().UnregisterInputPreProcessor(CommonInputPreprocessor);
+	}
 	CommonInputPreprocessor.Reset();
 
 	FTSTicker::GetCoreTicker().RemoveTicker(TickHandle);
diff --git a/Engine/Plugins/Experimental/CommonUI/Source/CommonUI/Private/Input/CommonUIActionRouterBase.cpp b/Engine/Plugins/Experimental/CommonUI/Source/CommonUI/Private/Input/CommonUIActionRouterBase.cpp
--- a/Engine/Plugins/Experimental/CommonUI/Source/CommonUI/Private/Input/CommonUIActionRouterBase.cpp
+++ b/Engine/Plugins/Experimental/CommonUI/Source/CommonUI/Private/Input/CommonUIActionRouterBase.cpp
@@ -229,8 +229,11 @@ void UCommonUIActionRouterBase::Deinitialize()
 {
 	Super::Deinitialize();
 	
-	FSlateApplication::Get().OnFocusChanging().RemoveAll(this);
-	FSlateApplication::Get().UnregisterInputPreProcessor(AnalogCursor);
+	if (FSlateApplication::IsInitialized())
+	{
+		FSlateApplication::Get().OnFocusChanging().RemoveAll(this);
+		FSlateApplication::Get().UnregisterInputPreProcessor(AnalogCursor);
+	}
 	FTSTicker::GetCoreTicker().RemoveTicker(TickHandle);
 	SetActiveRoot(nullptr);
 	HeldKeys.Empty();
