Temp fix prevent crash on CommonUI shutdown.

diff --git a/Engine/Plugins/Experimental/CommonUI/Source/CommonUI/Private/Input/CommonUIActionRouterBase.cpp b/Engine/Plugins/Experimental/CommonUI/Source/CommonUI/Private/Input/CommonUIActionRouterBase.cpp
--- a/Engine/Plugins/Experimental/CommonUI/Source/CommonUI/Private/Input/CommonUIActionRouterBase.cpp
+++ b/Engine/Plugins/Experimental/CommonUI/Source/CommonUI/Private/Input/CommonUIActionRouterBase.cpp
@@ -538,7 +538,11 @@ void UCommonUIActionRouterBase::HandleRootWidgetSlateReleased(UCommonActivatable
 		// the deactivation. Not a big deal, just need to process the deactivation right here if the node in question is the active root.
 		if (ActiveRootNode == RootNodes[RootToRemoveIdx])
 		{
-			check(!ActiveRootNode->IsWidgetActivated());
+			// @TODO: DarenC - Temp fix for crash on PIE exit
+			if (!ensure(!ActiveRootNode->IsWidgetActivated()) && ActiveRootNode->GetWidget())
+			{
+				ActiveRootNode->GetWidget()->DeactivateWidget();
+			}
 			HandleRootNodeDeactivated(ActiveRootNode);
 		}
 
