Prevent SCommonTileView from stealing focus after scrolling complete when focus is lost.
#tests: PIE: Verified focus is no longer stolen back when given to the SCommonTileView then immediately given to an external widget.

diff --git a/Engine/Plugins/Experimental/CommonUI/Source/CommonUI/Private/CommonTileView.cpp b/Engine/Plugins/Experimental/CommonUI/Source/CommonUI/Private/CommonTileView.cpp
--- a/Engine/Plugins/Experimental/CommonUI/Source/CommonUI/Private/CommonTileView.cpp
+++ b/Engine/Plugins/Experimental/CommonUI/Source/CommonUI/Private/CommonTileView.cpp
@@ -12,7 +12,7 @@ template <typename ItemType>
 class SCommonTileView : public STileView<ItemType>
 {
 public:
-	FReply OnFocusReceived(const FGeometry& MyGeometry, const FFocusEvent& InFocusEvent)
+	virtual FReply OnFocusReceived(const FGeometry& MyGeometry, const FFocusEvent& InFocusEvent) override
 	{
 		if (bScrollToSelectedOnFocus && (InFocusEvent.GetCause() == EFocusCause::Navigation || InFocusEvent.GetCause() == EFocusCause::SetDirectly))
 		{
@@ -39,6 +39,12 @@ public:
 		return STileView<ItemType>::OnFocusReceived(MyGeometry, InFocusEvent);
 	}
 
+	virtual void OnFocusLost(const FFocusEvent& InFocusEvent) override
+	{
+		// if we've lost focus, even if we're already scrolling a tile into view, refrain from navigating to that tile and thus stealing focus from another widget
+		this->bNavigateOnScrollIntoView = false;
+	}
+
 	virtual void OnMouseLeave(const FPointerEvent& MouseEvent) override
 	{
 		STileView<ItemType>::OnMouseLeave(MouseEvent);
