Avoid changing cursor visibility in editor if game viewport is not focused.

diff --git a/Engine/Plugins/Experimental/CommonUI/Source/CommonUI/Private/Input/CommonAnalogCursor.cpp b/Engine/Plugins/Experimental/CommonUI/Source/CommonUI/Private/Input/CommonAnalogCursor.cpp
--- a/Engine/Plugins/Experimental/CommonUI/Source/CommonUI/Private/Input/CommonAnalogCursor.cpp
+++ b/Engine/Plugins/Experimental/CommonUI/Source/CommonUI/Private/Input/CommonAnalogCursor.cpp
@@ -311,6 +311,24 @@ bool FCommonAnalogCursor::HandleAnalogInputEvent(FSlateApplication& SlateApp, co
 	return false;
 }
 
+bool FCommonAnalogCursor::HandleMouseMoveEvent(FSlateApplication& SlateApp, const FPointerEvent& MouseEvent)
+{
+#if WITH_EDITOR
+	// We can leave editor cursor visibility in a bad state if the engine stops ticking to debug
+	if (GIntraFrameDebuggingGameThread)
+	{
+		SlateApp.SetPlatformCursorVisibility(true);
+		TSharedPtr<FSlateUser> SlateUser = FSlateApplication::Get().GetUser(GetOwnerUserIndex());
+		if (ensure(SlateUser))
+		{
+			SlateUser->SetCursorVisibility(true);
+		}
+	}
+#endif // WITH_EDITOR
+
+	return FAnalogCursor::HandleMouseMoveEvent(SlateApp, MouseEvent);
+}
+
 bool FCommonAnalogCursor::HandleMouseButtonDownEvent(FSlateApplication& SlateApp, const FPointerEvent& PointerEvent)
 {
 	if (FAnalogCursor::IsRelevantInput(PointerEvent))
diff --git a/Engine/Plugins/Experimental/CommonUI/Source/CommonUI/Public/Input/CommonAnalogCursor.h b/Engine/Plugins/Experimental/CommonUI/Source/CommonUI/Public/Input/CommonAnalogCursor.h
--- a/Engine/Plugins/Experimental/CommonUI/Source/CommonUI/Public/Input/CommonAnalogCursor.h
+++ b/Engine/Plugins/Experimental/CommonUI/Source/CommonUI/Public/Input/CommonAnalogCursor.h
@@ -40,6 +40,7 @@ public:
 	bool CanReleaseMouseCapture() const;
 
 	virtual bool HandleAnalogInputEvent(FSlateApplication& SlateApp, const FAnalogInputEvent& InAnalogInputEvent) override;
+	virtual bool HandleMouseMoveEvent(FSlateApplication& SlateApp, const FPointerEvent& MouseEvent) override;
 	virtual bool HandleMouseButtonDownEvent(FSlateApplication& SlateApp, const FPointerEvent& InPointerEvent) override;
 	virtual bool HandleMouseButtonUpEvent(FSlateApplication& SlateApp, const FPointerEvent& InPointerEvent) override;
 
