ActionDomains - A system for managing input event flows in CommonUIActionRouter

diff --git a/Engine/Plugins/Experimental/CommonUI/Source/CommonInput/Private/CommonInputActionDomain.cpp b/Engine/Plugins/Experimental/CommonUI/Source/CommonInput/Private/CommonInputActionDomain.cpp
new file mode 100644
--- /dev/null
+++ b/Engine/Plugins/Experimental/CommonUI/Source/CommonInput/Private/CommonInputActionDomain.cpp
@@ -0,0 +1,45 @@
+// Copyright Epic Games, Inc. All Rights Reserved.
+
+#include "CommonInputActionDomain.h"
+
+bool UCommonInputActionDomain::ShouldBreakInnerEventFlow(bool bInputEventHandled) const
+{
+	switch (InnerBehavior)
+	{
+		case ECommonInputEventFlowBehavior::BlockIfActive:
+		{
+			return true;
+		}
+		case ECommonInputEventFlowBehavior::BlockIfHandled:
+		{
+			return bInputEventHandled;
+		}
+		case ECommonInputEventFlowBehavior::NeverBlock:
+		{
+			return false;
+		}
+	}
+
+	return false;
+}
+
+bool UCommonInputActionDomain::ShouldBreakEventFlow(bool bDomainHadActiveRoots, bool bInputEventHandledAtLeastOnce) const
+{
+	switch (Behavior)
+	{
+		case ECommonInputEventFlowBehavior::BlockIfActive:
+		{
+			return bDomainHadActiveRoots;
+		}
+		case ECommonInputEventFlowBehavior::BlockIfHandled:
+		{
+			return bInputEventHandledAtLeastOnce;
+		}
+		case ECommonInputEventFlowBehavior::NeverBlock:
+		{
+			return false;
+		}
+	}
+
+	return false;
+}
\ No newline at end of file
diff --git a/Engine/Plugins/Experimental/CommonUI/Source/CommonInput/Private/CommonInputSettings.cpp b/Engine/Plugins/Experimental/CommonUI/Source/CommonInput/Private/CommonInputSettings.cpp
--- a/Engine/Plugins/Experimental/CommonUI/Source/CommonInput/Private/CommonInputSettings.cpp
+++ b/Engine/Plugins/Experimental/CommonUI/Source/CommonInput/Private/CommonInputSettings.cpp
@@ -6,6 +6,7 @@
 #include "Engine/StreamableManager.h"
 #include "Engine/AssetManager.h"
 #include "Misc/DataDrivenPlatformInfoRegistry.h"
+#include "CommonInputActionDomain.h"
 #include "CommonInputBaseTypes.h"
 #include "Engine/PlatformSettings.h"
 
@@ -21,6 +22,7 @@ UCommonInputSettings::UCommonInputSettings(const FObjectInitializer& Initializer
 void UCommonInputSettings::LoadData()
 {
 	LoadInputData();
+	LoadActionDomainTable();
 }
 
 #if WITH_EDITOR
@@ -64,6 +66,27 @@ void UCommonInputSettings::LoadInputData()
 	}
 }
 
+void UCommonInputSettings::LoadActionDomainTable()
+{
+	if (!bActionDomainTableLoaded)
+	{
+		// If we were created early enough to be disregarded by the GC (which we should be), then we need to 
+		// add all of our members to the root set, since our hard reference to them is totally meaningless to the GC.
+		const bool bIsDisregardForGC = GUObjectArray.IsDisregardForGC(this);
+
+		ActionDomainTablePtr = ActionDomainTable.LoadSynchronous();
+		if (ActionDomainTablePtr)
+		{
+			if (bIsDisregardForGC)
+			{
+				ActionDomainTablePtr->AddToRoot();
+			}
+		}
+
+		bActionDomainTableLoaded = true;
+	}
+}
+
 void UCommonInputSettings::ValidateData()
 {
     bInputDataLoaded &= !InputData.IsPending();
diff --git a/Engine/Plugins/Experimental/CommonUI/Source/CommonInput/Private/CommonInputSubsystem.cpp b/Engine/Plugins/Experimental/CommonUI/Source/CommonInput/Private/CommonInputSubsystem.cpp
--- a/Engine/Plugins/Experimental/CommonUI/Source/CommonInput/Private/CommonInputSubsystem.cpp
+++ b/Engine/Plugins/Experimental/CommonUI/Source/CommonInput/Private/CommonInputSubsystem.cpp
@@ -323,6 +323,8 @@ void UCommonInputSubsystem::Initialize(FSubsystemCollectionBase& Collection)
 	TickHandle = FTSTicker::GetCoreTicker().AddTicker(FTickerDelegate::CreateUObject(this, &UCommonInputSubsystem::Tick), 0.1f);
 
 	CVarInputKeysVisible->SetOnChangedCallback(FConsoleVariableDelegate::CreateUObject(this, &UCommonInputSubsystem::ShouldShowInputKeysChanged));
+
+	SetActionDomainTable(FCommonInputBase::GetInputSettings()->GetActionDomainTable());
 }
 
 void UCommonInputSubsystem::Deinitialize()
diff --git a/Engine/Plugins/Experimental/CommonUI/Source/CommonInput/Public/CommonInputActionDomain.h b/Engine/Plugins/Experimental/CommonUI/Source/CommonInput/Public/CommonInputActionDomain.h
new file mode 100644
--- /dev/null
+++ b/Engine/Plugins/Experimental/CommonUI/Source/CommonInput/Public/CommonInputActionDomain.h
@@ -0,0 +1,54 @@
+// Copyright Epic Games, Inc. All Rights Reserved.
+
+#pragma once
+
+#include "Engine/DataAsset.h"
+
+#include "CommonInputActionDomain.generated.h"
+
+UENUM()
+enum class ECommonInputEventFlowBehavior {
+	BlockIfActive,
+	BlockIfHandled,
+	NeverBlock,
+};
+
+/**
+ * Describes an input-event handling domain. It's InnerBehavior determines how events
+ * flow between widgets within the domain and Behavior determines how events will flow to
+ * other Domains in the DomainTable.
+ */
+UCLASS()
+class COMMONINPUT_API UCommonInputActionDomain : public UDataAsset
+{
+	GENERATED_BODY()
+
+public:
+
+	// Behavior of an input event between Action Domains, i.e., how an event flows into the next Action Domain
+	UPROPERTY(EditDefaultsOnly, Category = "Default")
+	ECommonInputEventFlowBehavior Behavior = ECommonInputEventFlowBehavior::BlockIfActive;
+
+	// Behavior of an input event within an Action Domain, i.e., how an event flows to a lower ZOrder active widget
+	// within the same Action Domain
+	UPROPERTY(EditDefaultsOnly, Category = "Default")
+	ECommonInputEventFlowBehavior InnerBehavior = ECommonInputEventFlowBehavior::BlockIfHandled;
+
+	bool ShouldBreakInnerEventFlow(bool bInputEventHandled) const;
+
+	bool ShouldBreakEventFlow(bool bDomainHadActiveRoots, bool bInputEventHandledAtLeastOnce) const;
+};
+
+/**
+ * An ordered array of ActionDomains.
+ */
+UCLASS()
+class COMMONINPUT_API UCommonInputActionDomainTable : public UDataAsset
+{
+	GENERATED_BODY()
+
+public:
+	// Domains will receive events in ascending index order
+	UPROPERTY(EditDefaultsOnly, BlueprintReadOnly, Category = "Default")
+	TArray<UCommonInputActionDomain*> ActionDomains;
+};
\ No newline at end of file
diff --git a/Engine/Plugins/Experimental/CommonUI/Source/CommonInput/Public/CommonInputSettings.h b/Engine/Plugins/Experimental/CommonUI/Source/CommonInput/Public/CommonInputSettings.h
--- a/Engine/Plugins/Experimental/CommonUI/Source/CommonInput/Public/CommonInputSettings.h
+++ b/Engine/Plugins/Experimental/CommonUI/Source/CommonInput/Public/CommonInputSettings.h
@@ -9,6 +9,7 @@
 #include "InputCoreTypes.h"
 #include "Engine/StreamableManager.h"
 #include "Templates/SharedPointer.h"
+#include "CommonInputActionDomain.h"
 #include "CommonInputSubsystem.h"
 #include "CommonInputBaseTypes.h"
 
@@ -49,6 +50,8 @@ public:
 
 	bool GetEnableDefaultInputConfig() const { return bEnableDefaultInputConfig; }
 
+	TObjectPtr<UCommonInputActionDomainTable> GetActionDomainTable() const { return ActionDomainTablePtr; }
+
 private:
 	virtual void PostInitProperties() override;
 
@@ -84,11 +87,20 @@ private:
 	UPROPERTY(config, EditAnywhere, Category = "Input")
 	bool bEnableDefaultInputConfig = true;
 
+	/** Create a derived asset from UCommonInputActionDomainTable to store ordered ActionDomain data for your game */
+	UPROPERTY(config, EditAnywhere, Category = "Action Domain")
+	TSoftObjectPtr<UCommonInputActionDomainTable> ActionDomainTable;
+
 private:
 	void LoadInputData();
+	void LoadActionDomainTable();
 
 	bool bInputDataLoaded;
+	bool bActionDomainTableLoaded;
 
 	UPROPERTY(Transient)
 	TSubclassOf<UCommonUIInputData> InputDataClass;
+
+	UPROPERTY(Transient)
+	TObjectPtr<UCommonInputActionDomainTable> ActionDomainTablePtr;
 };
diff --git a/Engine/Plugins/Experimental/CommonUI/Source/CommonInput/Public/CommonInputSubsystem.h b/Engine/Plugins/Experimental/CommonUI/Source/CommonInput/Public/CommonInputSubsystem.h
--- a/Engine/Plugins/Experimental/CommonUI/Source/CommonInput/Public/CommonInputSubsystem.h
+++ b/Engine/Plugins/Experimental/CommonUI/Source/CommonInput/Public/CommonInputSubsystem.h
@@ -12,6 +12,7 @@ class UWidget;
 class ULocalPlayer;
 class APlayerController;
 class FCommonInputPreprocessor;
+class UCommonInputActionDomainTable;
 
 DECLARE_DYNAMIC_MULTICAST_DELEGATE_OneParam(FInputMethodChangedDelegate, ECommonInputType, bNewInputType);
 DECLARE_MULTICAST_DELEGATE_ThreeParams(FPlatformInputSupportOverrideDelegate, ULocalPlayer*, ECommonInputType, bool&);
@@ -68,6 +69,10 @@ public:
 	UFUNCTION(BlueprintCallable, Category = CommonInputSubsystem)
 	bool ShouldShowInputKeys() const;
 
+	void SetActionDomainTable(TObjectPtr<UCommonInputActionDomainTable> Table) { ActionDomainTable = Table; }
+
+	TObjectPtr<UCommonInputActionDomainTable> GetActionDomainTable() const { return ActionDomainTable; }
+
 	/** Returns true if the specified key can be present on both a mobile device and mobile gamepads */
 	static bool IsMobileGamepadKey(const FKey& InKey);
 
@@ -145,6 +150,9 @@ private:
 
 	TSharedPtr<FCommonInputPreprocessor> CommonInputPreprocessor;
 
+	UPROPERTY(Transient)
+	TObjectPtr<UCommonInputActionDomainTable> ActionDomainTable;
+
 	/** Is the current click simulated by the gamepad's face button down/right (platform dependent) */
 	UPROPERTY(Transient)
 	bool bIsGamepadSimulatedClick;
diff --git a/Engine/Plugins/Experimental/CommonUI/Source/CommonUI/Private/CommonActivatableWidget.cpp b/Engine/Plugins/Experimental/CommonUI/Source/CommonUI/Private/CommonActivatableWidget.cpp
--- a/Engine/Plugins/Experimental/CommonUI/Source/CommonUI/Private/CommonActivatableWidget.cpp
+++ b/Engine/Plugins/Experimental/CommonUI/Source/CommonUI/Private/CommonActivatableWidget.cpp
@@ -2,6 +2,7 @@
 
 #include "CommonActivatableWidget.h"
 #include "CommonUIPrivatePCH.h"
+#include "Input/CommonUIActionRouterBase.h"
 #include "Input/CommonUIInputTypes.h"
 #include "Input/UIActionRouterTypes.h"
 #include "ICommonInputModule.h"
@@ -146,6 +147,29 @@ void UCommonActivatableWidget::ClearActiveHoldInputs()
 	}
 }
 
+UCommonInputActionDomain* UCommonActivatableWidget::GetCalculatedActionDomain()
+{
+	UCommonInputActionDomain* CalculatedActionDomain = ActionDomain;
+
+	if (!bInheritActionDomain)
+	{
+		return CalculatedActionDomain;
+	}
+
+	const UCommonActivatableWidget* CurrentWidget = this;
+	while (CurrentWidget && CurrentWidget->bInheritActionDomain)
+	{
+		CurrentWidget = UCommonUIActionRouterBase::FindOwningActivatable(CurrentWidget->GetCachedWidget(), GetOwningLocalPlayer());
+	}
+
+	if (CurrentWidget && !CurrentWidget->bInheritActionDomain)
+	{
+		return CurrentWidget->ActionDomain;
+	}
+
+	return nullptr;
+}
+
 TSharedRef<SWidget> UCommonActivatableWidget::RebuildWidget()
 {
 	// Note: the scoped builder guards against design-time so we don't need to here (as it'd make the scoped lifetime more awkward to leverage)
diff --git a/Engine/Plugins/Experimental/CommonUI/Source/CommonUI/Private/Input/CommonUIActionRouterBase.cpp b/Engine/Plugins/Experimental/CommonUI/Source/CommonUI/Private/Input/CommonUIActionRouterBase.cpp
--- a/Engine/Plugins/Experimental/CommonUI/Source/CommonUI/Private/Input/CommonUIActionRouterBase.cpp
+++ b/Engine/Plugins/Experimental/CommonUI/Source/CommonUI/Private/Input/CommonUIActionRouterBase.cpp
@@ -1,10 +1,11 @@
 // Copyright Epic Games, Inc. All Rights Reserved.
 
 #include "Input/CommonUIActionRouterBase.h"
-#include "Input/UIActionRouterTypes.h"
 #include "Input/CommonAnalogCursor.h"
 #include "Input/CommonUIInputSettings.h"
+#include "Input/UIActionRouterTypes.h"
 #include "CommonInputSettings.h"
+#include "CommonInputActionDomain.h"
 #include "ICommonInputModule.h"
 
 #include "Framework/Application/SlateUser.h"
@@ -30,6 +31,12 @@ static const FAutoConsoleVariableRef CVarAlwaysShowCursor(
 	bAlwaysShowCursor,
 	TEXT(""));
 
+bool bEnableActionDomainRouting = false;
+static const FAutoConsoleVariableRef CVarEnableActionDomainRouting(
+	TEXT("CommonUI.EnableActionDomainRouting"),
+	bEnableActionDomainRouting,
+	TEXT(""));
+
 //@todo DanH: TEMP LOCATION
 FGlobalUITags FGlobalUITags::GUITags;
 
@@ -347,18 +354,35 @@ ERouteUIInputResult UCommonUIActionRouterBase::ProcessInput(FKey Key, EInputEven
 	// Begin with a pass to see if the input corresponds to a hold action
 	// We do this first to make sure that a higher-priority press binding doesn't prevent a hold on the same key from being triggerable
 	EProcessHoldActionResult ProcessHoldResult = PersistentActions->ProcessHoldInput(ActiveMode, Key, InputEvent);
-	if (bIsActivatableTreeEnabled && ActiveRootNode && ProcessHoldResult == EProcessHoldActionResult::Unhandled)
+	
+	if (ProcessHoldResult == EProcessHoldActionResult::Unhandled)
 	{
-		ProcessHoldResult = ActiveRootNode->ProcessHoldInput(ActiveMode, Key, InputEvent);
+		if (bEnableActionDomainRouting)
+		{
+			ProcessHoldResult = ProcessHoldInputOnActionDomains(ActiveMode, Key, InputEvent);
+		}
+		else if (bIsActivatableTreeEnabled && ActiveRootNode)
+		{
+			ProcessHoldResult = ActiveRootNode->ProcessHoldInput(ActiveMode, Key, InputEvent);
+		}
 	}
 
 	const auto ProcessNormalInputFunc = [Key, ActiveMode, this](EInputEvent Event)
 		{
 			bool bHandled = PersistentActions->ProcessNormalInput(ActiveMode, Key, Event);
-			if (!bHandled && ActiveRootNode && bIsActivatableTreeEnabled)
+
+			if (!bHandled)
 			{
-				bHandled = ActiveRootNode->ProcessNormalInput(ActiveMode, Key, Event);
+				if (bEnableActionDomainRouting)
+				{
+					bHandled = ProcessInputOnActionDomains(ActiveMode, Key, Event);
+				}
+				else if (ActiveRootNode && bIsActivatableTreeEnabled)
+				{
+					bHandled = ActiveRootNode->ProcessNormalInput(ActiveMode, Key, Event);
+				}
 			}
+
 			return bHandled;
 		};
 
@@ -547,6 +571,12 @@ void UCommonUIActionRouterBase::HandleRootWidgetSlateReleased(UCommonActivatable
 		}
 
 		TWeakPtr<FActivatableTreeRoot> ToBeRemoved = RootNodes[RootToRemoveIdx];
+
+		if (bEnableActionDomainRouting)
+		{
+			RemoveFromActionDomain(RootNodes[RootToRemoveIdx]);
+		}
+
 		RootNodes.RemoveAtSwap(RootToRemoveIdx);
 		// Cannot actually have this ensure here, because we may be in a function called on the ToBeRemoved node itself, keeping one remaining strong reference
 		// ensureAlways(!ToBeRemoved.IsValid());
@@ -614,11 +644,134 @@ void UCommonUIActionRouterBase::HandlePostGarbageCollect()
 	{
 		if (!Iter->Get().IsWidgetValid())
 		{
+			RemoveFromActionDomain(*Iter);
 			Iter.RemoveCurrent();
 		}
 	}
 }
 
+void UCommonUIActionRouterBase::AddToActionDomain(FActivatableTreeRootRef RootNode)
+{
+	UCommonInputActionDomain* ActionDomain = RootNode->GetWidget()->GetCalculatedActionDomain();
+
+	if (ActionDomain)
+	{
+		ActionDomainRootNodes.FindOrAdd(ActionDomain).Add(RootNode);
+	}
+	else
+	{
+		UE_LOG(LogUIActionRouter, Error, TEXT("ActionDomain of CommonActivatableWidget could not be resolved. Widget: %s"), *RootNode->GetWidget()->GetName());
+	}
+}
+
+void UCommonUIActionRouterBase::RemoveFromActionDomain(FActivatableTreeRootRef RootNode)
+{
+	UCommonInputActionDomain* ActionDomain = RootNode->GetWidget()->GetCalculatedActionDomain();
+
+	if (ActionDomain)
+	{
+		FActionDomainSortedRootList* ActionDomainRootList = ActionDomainRootNodes.Find(ActionDomain);
+
+		if (ensure(ActionDomainRootList))
+		{
+			verify(ActionDomainRootList->Remove(RootNode) != INDEX_NONE);
+		}
+	}
+	else
+	{
+		UE_LOG(LogUIActionRouter, Error, TEXT("ActionDomain of CommonActivatableWidget could not be resolved. Widget: %s"), *RootNode->GetWidget()->GetName());
+	}
+}
+
+bool UCommonUIActionRouterBase::ProcessInputOnActionDomains(ECommonInputMode ActiveInputMode, FKey Key, EInputEvent InputEvent) const
+{
+	UCommonInputSubsystem& CommonInputSubsystem = GetInputSubsystem();
+	UCommonInputActionDomainTable* ActionDomainTable = CommonInputSubsystem.GetActionDomainTable();
+
+	bool bInputEventHandledAtLeastOnce = false;
+
+	for (const UCommonInputActionDomain* ActionDomain : ActionDomainTable->ActionDomains)
+	{
+		const FActionDomainSortedRootList* SortedRootList = ActionDomainRootNodes.Find(ActionDomain);
+		if (!SortedRootList)
+		{
+			// No widget with this Domain was added
+			continue;
+		}
+
+		bool bInputEventHandledInDomain = false;
+		bool bDomainHadActiveRoots = false;
+
+		for (const FActivatableTreeRootRef& RootNode : SortedRootList->RootList)
+		{
+			if (RootNode->IsWidgetActivated())
+			{
+				bool bInputEventHandled = RootNode->ProcessActionDomainNormalInput(ActiveInputMode, Key, InputEvent);
+				bInputEventHandledInDomain |= bInputEventHandled;
+				bDomainHadActiveRoots = true;
+
+				if (ActionDomain->ShouldBreakInnerEventFlow(bInputEventHandled))
+				{
+					break;
+				}
+			}
+		}
+
+		bInputEventHandledAtLeastOnce |= bInputEventHandledInDomain;
+
+		if (ActionDomain->ShouldBreakEventFlow(bDomainHadActiveRoots, bInputEventHandledInDomain))
+		{
+			break;
+		}
+	}
+
+	return bInputEventHandledAtLeastOnce;
+}
+
+EProcessHoldActionResult UCommonUIActionRouterBase::ProcessHoldInputOnActionDomains(ECommonInputMode ActiveInputMode, FKey Key, EInputEvent InputEvent) const
+{
+	UCommonInputSubsystem& CommonInputSubsystem = GetInputSubsystem();
+	UCommonInputActionDomainTable* ActionDomainTable = CommonInputSubsystem.GetActionDomainTable();
+
+	EProcessHoldActionResult OutHoldActionResult = EProcessHoldActionResult::Unhandled;
+
+	for (const UCommonInputActionDomain* ActionDomain : ActionDomainTable->ActionDomains)
+	{
+		const FActionDomainSortedRootList* SortedRootList = ActionDomainRootNodes.Find(ActionDomain);
+		if (!SortedRootList)
+		{
+			// No widget with this Domain was added
+			continue;
+		}
+
+		bool bInputEventHandledInDomain = false;
+		bool bDomainHadActiveRoots = false;
+
+		for (const FActivatableTreeRootRef& RootNode : SortedRootList->RootList)
+		{
+			if (RootNode->IsWidgetActivated())
+			{
+				bool bInputEventHandled = RootNode->ProcessActionDomainHoldInput(ActiveInputMode, Key, InputEvent, OutHoldActionResult);
+				bInputEventHandledInDomain |= bInputEventHandled;
+				bDomainHadActiveRoots = true;
+
+				if (ActionDomain->ShouldBreakInnerEventFlow(bInputEventHandled))
+				{
+					break;
+				}
+			}
+
+		}
+
+		if (ActionDomain->ShouldBreakEventFlow(bDomainHadActiveRoots, bInputEventHandledInDomain))
+		{
+			break;
+		}
+	}
+
+	return OutHoldActionResult;
+}
+
 void UCommonUIActionRouterBase::ProcessRebuiltWidgets()
 {
 	// Begin by organizing all of the widgets that need nodes according to their direct parent
@@ -651,6 +804,11 @@ void UCommonUIActionRouterBase::ProcessRebuiltWidgets()
 		RootWidget->OnSlateReleased().AddUObject(this, &UCommonUIActionRouterBase::HandleRootWidgetSlateReleased, RootWidget);
 		RootNodes.Add(RootNode);
 
+		if (bEnableActionDomainRouting)
+		{
+			AddToActionDomain(RootNode);
+		}
+
 		AssembleTreeRecursive(RootNode, WidgetsByDirectParent);
 
 		if (RootWidget->IsActivated())
@@ -1281,3 +1439,24 @@ static const FAutoConsoleCommandWithWorld DumpInputConfigCommand(
 	TEXT("Outputs the current Input Config for each player"),
 	FConsoleCommandWithWorldDelegate::CreateStatic(&FActionRouterDebugUtils::HandleDumpCurrentInputConfig)
 );
+
+void UCommonUIActionRouterBase::FActionDomainSortedRootList::Add(FActivatableTreeRootRef RootNode)
+{
+	auto SortFunc = [](const FActivatableTreeRootRef& A, const FActivatableTreeRootRef& B) 
+		{
+			return A->GetLastPaintLayer() < B->GetLastPaintLayer();
+		};
+
+	RootList.Add(RootNode);
+	RootList.Sort(SortFunc);
+}
+
+int32 UCommonUIActionRouterBase::FActionDomainSortedRootList::Remove(FActivatableTreeRootRef RootNode)
+{
+	return RootList.Remove(RootNode);
+}
+
+bool UCommonUIActionRouterBase::FActionDomainSortedRootList::Contains(FActivatableTreeRootRef RootNode)
+{
+	return RootList.Contains(RootNode);
+}
diff --git a/Engine/Plugins/Experimental/CommonUI/Source/CommonUI/Private/Input/UIActionRouterTypes.cpp b/Engine/Plugins/Experimental/CommonUI/Source/CommonUI/Private/Input/UIActionRouterTypes.cpp
--- a/Engine/Plugins/Experimental/CommonUI/Source/CommonUI/Private/Input/UIActionRouterTypes.cpp
+++ b/Engine/Plugins/Experimental/CommonUI/Source/CommonUI/Private/Input/UIActionRouterTypes.cpp
@@ -712,6 +712,21 @@ EProcessHoldActionResult FActivatableTreeNode::ProcessHoldInput(ECommonInputMode
 	return EProcessHoldActionResult::Unhandled;
 }
 
+bool FActivatableTreeNode::ProcessActionDomainHoldInput(ECommonInputMode ActiveInputMode, FKey Key, EInputEvent InputEvent, EProcessHoldActionResult& OutHoldActionResult) const
+{
+	for (const FActivatableTreeNodeRef& ChildNode : Children)
+	{
+		EProcessHoldActionResult ChildResult = ChildNode->ProcessHoldInput(ActiveInputMode, Key, InputEvent);
+		if (ChildResult != EProcessHoldActionResult::Unhandled)
+		{
+			OutHoldActionResult = ChildResult;
+			return true;
+		}
+	}
+	OutHoldActionResult = FActionRouterBindingCollection::ProcessHoldInput(ActiveInputMode, Key, InputEvent);
+	return OutHoldActionResult != EProcessHoldActionResult::Unhandled;
+}
+
 bool FActivatableTreeNode::ProcessNormalInput(ECommonInputMode ActiveInputMode, FKey Key, EInputEvent InputEvent) const
 {
 	if (IsReceivingInput())
@@ -728,6 +743,18 @@ bool FActivatableTreeNode::ProcessNormalInput(ECommonInputMode ActiveInputMode,
 	return false;
 }
 
+bool FActivatableTreeNode::ProcessActionDomainNormalInput(ECommonInputMode ActiveInputMode, FKey Key, EInputEvent InputEvent) const
+{
+	for (const FActivatableTreeNodeRef& ChildNode : Children)
+	{
+		if (ChildNode->ProcessActionDomainNormalInput(ActiveInputMode, Key, InputEvent))
+		{
+			return true;
+		}
+	}
+	return FActionRouterBindingCollection::ProcessNormalInput(ActiveInputMode, Key, InputEvent);
+}
+
 bool FActivatableTreeNode::IsWidgetValid() const
 { 
 	return RepresentedWidget.IsValid();
diff --git a/Engine/Plugins/Experimental/CommonUI/Source/CommonUI/Private/Input/UIActionRouterTypes.h b/Engine/Plugins/Experimental/CommonUI/Source/CommonUI/Private/Input/UIActionRouterTypes.h
--- a/Engine/Plugins/Experimental/CommonUI/Source/CommonUI/Private/Input/UIActionRouterTypes.h
+++ b/Engine/Plugins/Experimental/CommonUI/Source/CommonUI/Private/Input/UIActionRouterTypes.h
@@ -162,9 +162,11 @@ class COMMONUI_API FActivatableTreeNode : public FActionRouterBindingCollection
 {
 public:
 	virtual ~FActivatableTreeNode();
-
+	
 	virtual EProcessHoldActionResult ProcessHoldInput(ECommonInputMode ActiveInputMode, FKey Key, EInputEvent InputEvent) const override;
+	virtual bool ProcessActionDomainHoldInput(ECommonInputMode ActiveInputMode, FKey Key, EInputEvent InputEvent, EProcessHoldActionResult& OutHoldActionResult) const;
 	virtual bool ProcessNormalInput(ECommonInputMode ActiveInputMode, FKey Key, EInputEvent InputEvent) const override;
+	virtual bool ProcessActionDomainNormalInput(ECommonInputMode ActiveInputMode, FKey Key, EInputEvent InputEvent) const;
 	virtual bool IsReceivingInput() const override { return bCanReceiveInput && IsWidgetActivated(); }
 
 	bool IsWidgetValid() const;
diff --git a/Engine/Plugins/Experimental/CommonUI/Source/CommonUI/Public/CommonActivatableWidget.h b/Engine/Plugins/Experimental/CommonUI/Source/CommonUI/Public/CommonActivatableWidget.h
--- a/Engine/Plugins/Experimental/CommonUI/Source/CommonUI/Public/CommonActivatableWidget.h
+++ b/Engine/Plugins/Experimental/CommonUI/Source/CommonUI/Public/CommonActivatableWidget.h
@@ -4,9 +4,11 @@
 
 #include "CommonUserWidget.h"
 #include "Input/UIActionBindingHandle.h"
+#include "CommonInputActionDomain.h"
 #include "CommonActivatableWidget.generated.h"
 
 class FActivatableTreeNode;
+class UCommonInputActionDomain;
 
 DECLARE_DYNAMIC_MULTICAST_DELEGATE(FOnWidgetActivationChanged);
 
@@ -96,6 +98,11 @@ public:
 	void RegisterInputTreeNode(const TSharedPtr<FActivatableTreeNode>& OwnerNode);
 	void ClearActiveHoldInputs();
 
+	/**
+	 * Returns the widget's ActionDomain, respecting any inheritance requirements.
+	 */
+	UCommonInputActionDomain* GetCalculatedActionDomain();
+
 protected:
 	virtual TSharedRef<SWidget> RebuildWidget() override;
 	virtual void ReleaseSlateResources(bool bReleaseChildren) override;
@@ -177,7 +184,16 @@ protected:
 	 */
 	UPROPERTY(EditAnywhere, Category = Activation, meta = (EditCondition = bSupportsActivationFocus))
 	bool bAutoRestoreFocus = false;
-	
+
+	UPROPERTY(EditAnywhere, BlueprintReadOnly, Category = "Input | ActionDomain", meta = (InlineEditConditionToggle))
+	bool bInheritActionDomain = true;
+
+	/**
+	 * Disable to inherit from owning CommonActivatableWidget.
+	 */
+	UPROPERTY(EditAnywhere, BlueprintReadOnly, Category = "Input | ActionDomain", meta = (EditCondition = "!bInheritActionDomain"))
+	TObjectPtr<UCommonInputActionDomain> ActionDomain;
+
 private:
 	/** See BindVisibilityToMultipleActivations */
 	void HandleVisibilityBoundWidgetActivations();
diff --git a/Engine/Plugins/Experimental/CommonUI/Source/CommonUI/Public/Input/CommonUIActionRouterBase.h b/Engine/Plugins/Experimental/CommonUI/Source/CommonUI/Public/Input/CommonUIActionRouterBase.h
--- a/Engine/Plugins/Experimental/CommonUI/Source/CommonUI/Public/Input/CommonUIActionRouterBase.h
+++ b/Engine/Plugins/Experimental/CommonUI/Source/CommonUI/Public/Input/CommonUIActionRouterBase.h
@@ -21,8 +21,10 @@ class UPlayerInput;
 class FCommonAnalogCursor;
 class IInputProcessor;
 
+enum class EProcessHoldActionResult;
 class FActivatableTreeNode;
 class UCommonInputSubsystem;
+class UCommonInputActionDomain;
 struct FUIActionBinding;
 class FDebugDisplayInfo;
 struct FAutoCompleteCommand;
@@ -157,6 +159,11 @@ private:
 
 	void HandlePostGarbageCollect();
 
+	void AddToActionDomain(FActivatableTreeRootRef RootNode);
+	void RemoveFromActionDomain(FActivatableTreeRootRef RootNode);
+	bool ProcessInputOnActionDomains(ECommonInputMode ActiveInputMode, FKey Key, EInputEvent InputEvent) const;
+	EProcessHoldActionResult ProcessHoldInputOnActionDomains(ECommonInputMode ActiveInputMode, FKey Key, EInputEvent InputEvent) const;
+
 	struct FPendingWidgetRegistration
 	{
 		TWeakObjectPtr<const UWidget> Widget;
@@ -200,4 +207,21 @@ private:
 	friend class FActionRouterDebugUtils;
 
 	mutable TArray<FKey> HeldKeys;
+
+	/** A wrapper around TArray that keeps RootList sorted by PaintLayer during insertion */
+	struct FActionDomainSortedRootList 
+	{
+		TArray<FActivatableTreeRootRef> RootList;
+
+		// Inserts RootNode into RootList based on its paint layer
+		void Add(FActivatableTreeRootRef RootNode);
+
+		// Trivial removal
+		int32 Remove(FActivatableTreeRootRef RootNode);
+
+		// Trivial Contains check
+		bool Contains(FActivatableTreeRootRef RootNode);
+	};
+
+	TMap<TObjectPtr<UCommonInputActionDomain>, FActionDomainSortedRootList> ActionDomainRootNodes;
 };
