Fix to ensure the mobile multiplier for UCommonRichTextBlock and UCommonTextBlock  is correctly applied in-game
- Removed IsDesign check from the SynchronizeProperties function of UCommonRichTextBlock that was preventing the MobileTextBlockScale from being applied in-game
- Reapplied MobileFontSizeMultiplier in the SetText function of UCommonTextBlock so that text block's behavior matches that of the UCommonRichTextBlock

diff --git a/Engine/Plugins/Experimental/CommonUI/Source/CommonUI/Private/CommonRichTextBlock.cpp b/Engine/Plugins/Experimental/CommonUI/Source/CommonUI/Private/CommonRichTextBlock.cpp
--- a/Engine/Plugins/Experimental/CommonUI/Source/CommonUI/Private/CommonRichTextBlock.cpp
+++ b/Engine/Plugins/Experimental/CommonUI/Source/CommonUI/Private/CommonRichTextBlock.cpp
@@ -582,12 +582,9 @@ void UCommonRichTextBlock::SynchronizeProperties()
 
 	if (MyRichTextBlock.IsValid())
 	{
-		if (IsDesignTime())
+		if (CommonUIUtils::ShouldDisplayMobileUISizes() && DefaultTextStyleOverrideClass.GetDefaultObject() == nullptr)
 		{
-			if (CommonUIUtils::ShouldDisplayMobileUISizes() && DefaultTextStyleOverrideClass.GetDefaultObject() == nullptr)
-			{
-				ApplyTextBlockScale();
-			}
+			ApplyTextBlockScale();
 		}
 	}
 }
diff --git a/Engine/Plugins/Experimental/CommonUI/Source/CommonUI/Private/CommonTextBlock.cpp b/Engine/Plugins/Experimental/CommonUI/Source/CommonUI/Private/CommonTextBlock.cpp
--- a/Engine/Plugins/Experimental/CommonUI/Source/CommonUI/Private/CommonTextBlock.cpp
+++ b/Engine/Plugins/Experimental/CommonUI/Source/CommonUI/Private/CommonTextBlock.cpp
@@ -420,12 +420,9 @@ void UCommonTextBlock::SynchronizeProperties()
 
 	Super::SynchronizeProperties();
 
-	if (CommonUIUtils::ShouldDisplayMobileUISizes() && MyTextBlock)
+	if (CommonUIUtils::ShouldDisplayMobileUISizes())
 	{
-		// Apply the font size multiplier
-		FSlateFontInfo FontInfo = Font;
-		FontInfo.Size *= MobileFontSizeMultiplier;
-		MyTextBlock->SetFont(FontInfo);
+		ApplyFontSizeMultiplier();
 	}
 }
 
@@ -437,6 +434,11 @@ void UCommonTextBlock::SetText(FText InText)
 	{
 		SetVisibility(InText.IsEmpty() ? ESlateVisibility::Collapsed : ESlateVisibility::SelfHitTestInvisible);
 	}
+
+	if (CommonUIUtils::ShouldDisplayMobileUISizes())
+	{
+		ApplyFontSizeMultiplier();
+	}
 }
 
 void UCommonTextBlock::UpdateFromStyle()
@@ -562,3 +564,13 @@ const UCommonTextScrollStyle* UCommonTextBlock::GetScrollStyleCDO() const
 {
 	return ScrollStyle ? Cast<UCommonTextScrollStyle>(ScrollStyle->ClassDefaultObject) : nullptr;
 }
+
+void UCommonTextBlock::ApplyFontSizeMultiplier() const
+{
+	if (MyTextBlock.IsValid())
+	{
+		FSlateFontInfo FontInfo = Font;
+		FontInfo.Size *= MobileFontSizeMultiplier;
+		MyTextBlock->SetFont(FontInfo);
+	}
+}
diff --git a/Engine/Plugins/Experimental/CommonUI/Source/CommonUI/Public/CommonTextBlock.h b/Engine/Plugins/Experimental/CommonUI/Source/CommonUI/Public/CommonTextBlock.h
--- a/Engine/Plugins/Experimental/CommonUI/Source/CommonUI/Public/CommonTextBlock.h
+++ b/Engine/Plugins/Experimental/CommonUI/Source/CommonUI/Public/CommonTextBlock.h
@@ -223,4 +223,6 @@ private:
 	const UCommonTextScrollStyle* GetScrollStyleCDO() const;
 
 	TSharedPtr<class STextScroller> TextScroller;
+
+	void ApplyFontSizeMultiplier() const;
 };
\ No newline at end of file
