bugfix: ActionDomain

In HandlePostGarbageCollect, the guard bEnableActionDomainRouting was not checked before accessing ActionDomain operations.

diff --git a/Engine/Plugins/Experimental/CommonUI/Source/CommonUI/Private/Input/CommonUIActionRouterBase.cpp b/Engine/Plugins/Experimental/CommonUI/Source/CommonUI/Private/Input/CommonUIActionRouterBase.cpp
--- a/Engine/Plugins/Experimental/CommonUI/Source/CommonUI/Private/Input/CommonUIActionRouterBase.cpp
+++ b/Engine/Plugins/Experimental/CommonUI/Source/CommonUI/Private/Input/CommonUIActionRouterBase.cpp
@@ -644,7 +644,10 @@ void UCommonUIActionRouterBase::HandlePostGarbageCollect()
 	{
 		if (!Iter->Get().IsWidgetValid())
 		{
-			RemoveFromActionDomain(*Iter);
+			if (bEnableActionDomainRouting)
+			{
+				RemoveFromActionDomain(*Iter);
+			}
 			Iter.RemoveCurrent();
 		}
 	}
