CommonUI - When Activatable panels request 'menu' mode for their input, we're still allowing the game to get access to the input.  The common game viewport client should be reporting it as handled, but wasn't, and so it would continue to pass to the game uninterrupted.

diff --git a/Engine/Plugins/Experimental/CommonUI/Source/CommonUI/Private/CommonGameViewportClient.cpp b/Engine/Plugins/Experimental/CommonUI/Source/CommonUI/Private/CommonGameViewportClient.cpp
--- a/Engine/Plugins/Experimental/CommonUI/Source/CommonUI/Private/CommonGameViewportClient.cpp
+++ b/Engine/Plugins/Experimental/CommonUI/Source/CommonUI/Private/CommonGameViewportClient.cpp
@@ -108,6 +108,9 @@ void UCommonGameViewportClient::HandleRerouteInput(int32 ControllerId, FKey Key,
 			ERouteUIInputResult InputResult = ActionRouter->ProcessInput(Key, EventType);
 			if (InputResult == ERouteUIInputResult::BlockGameInput)
 			{
+				// We need to set the reply as handled otherwise the input won't actually be blocked from reaching the viewport.
+				Reply = FReply::Handled();
+				// Notify interested parties that we blocked the input.
 				OnRerouteBlockedInput().ExecuteIfBound(ControllerId, Key, EventType, Reply);
 			}
 			else if (InputResult == ERouteUIInputResult::Handled)
