CommonActivatableWidget - ActionDomain pointer is now TSoftObjectPtr since it does not need to hold a hard reference. Changed Display Name to signify that you're overriding default inherit behavior by assigning an ActionDomain.

diff --git a/Engine/Plugins/Experimental/CommonUI/Source/CommonUI/Private/CommonActivatableWidget.cpp b/Engine/Plugins/Experimental/CommonUI/Source/CommonUI/Private/CommonActivatableWidget.cpp
--- a/Engine/Plugins/Experimental/CommonUI/Source/CommonUI/Private/CommonActivatableWidget.cpp
+++ b/Engine/Plugins/Experimental/CommonUI/Source/CommonUI/Private/CommonActivatableWidget.cpp
@@ -147,13 +147,11 @@ void UCommonActivatableWidget::ClearActiveHoldInputs()
 	}
 }
 
-UCommonInputActionDomain* UCommonActivatableWidget::GetCalculatedActionDomain()
+TObjectPtr<UCommonInputActionDomain> UCommonActivatableWidget::GetCalculatedActionDomain()
 {
-	UCommonInputActionDomain* CalculatedActionDomain = ActionDomain;
-
 	if (!bInheritActionDomain)
 	{
-		return CalculatedActionDomain;
+		return ActionDomain.Get();
 	}
 
 	const UCommonActivatableWidget* CurrentWidget = this;
@@ -164,7 +162,7 @@ UCommonInputActionDomain* UCommonActivatableWidget::GetCalculatedActionDomain()
 
 	if (CurrentWidget && !CurrentWidget->bInheritActionDomain)
 	{
-		return CurrentWidget->ActionDomain;
+		return CurrentWidget->ActionDomain.Get();
 	}
 
 	return nullptr;
diff --git a/Engine/Plugins/Experimental/CommonUI/Source/CommonUI/Public/CommonActivatableWidget.h b/Engine/Plugins/Experimental/CommonUI/Source/CommonUI/Public/CommonActivatableWidget.h
--- a/Engine/Plugins/Experimental/CommonUI/Source/CommonUI/Public/CommonActivatableWidget.h
+++ b/Engine/Plugins/Experimental/CommonUI/Source/CommonUI/Public/CommonActivatableWidget.h
@@ -101,7 +101,7 @@ public:
 	/**
 	 * Returns the widget's ActionDomain, respecting any inheritance requirements.
 	 */
-	UCommonInputActionDomain* GetCalculatedActionDomain();
+	TObjectPtr<UCommonInputActionDomain> GetCalculatedActionDomain();
 
 protected:
 	virtual TSharedRef<SWidget> RebuildWidget() override;
@@ -191,8 +191,8 @@ protected:
 	/**
 	 * Disable to inherit from owning CommonActivatableWidget.
 	 */
-	UPROPERTY(EditAnywhere, BlueprintReadOnly, Category = "Input | ActionDomain", meta = (EditCondition = "!bInheritActionDomain"))
-	TObjectPtr<UCommonInputActionDomain> ActionDomain;
+	UPROPERTY(EditAnywhere, BlueprintReadOnly, Category = "Input | ActionDomain", meta = (EditCondition = "!bInheritActionDomain", DisplayName = "Override ActionDomain"))
+	TSoftObjectPtr<UCommonInputActionDomain> ActionDomain;
 
 private:
 	/** See BindVisibilityToMultipleActivations */
