CommonUI - Crash fix - In UCommonHardwareVisibilityBorder, use PlatformTraits from UCommonUISettings if GetOwningLocalPlayer() fails

diff --git a/Engine/Plugins/Experimental/CommonUI/Source/CommonUI/Private/CommonHardwareVisibilityBorder.cpp b/Engine/Plugins/Experimental/CommonUI/Source/CommonUI/Private/CommonHardwareVisibilityBorder.cpp
--- a/Engine/Plugins/Experimental/CommonUI/Source/CommonUI/Private/CommonHardwareVisibilityBorder.cpp
+++ b/Engine/Plugins/Experimental/CommonUI/Source/CommonUI/Private/CommonHardwareVisibilityBorder.cpp
@@ -7,6 +7,8 @@
 #include "CommonInputBaseTypes.h"
 #include "Misc/DataDrivenPlatformInfoRegistry.h"
 #include "CommonUIVisibilitySubsystem.h"
+#include "CommonUISettings.h"
+#include "ICommonUIModule.h"
 
 UCommonHardwareVisibilityBorder::UCommonHardwareVisibilityBorder(const FObjectInitializer& ObjectInitializer)
 	: Super(ObjectInitializer)
@@ -27,18 +29,26 @@ void UCommonHardwareVisibilityBorder::UpdateVisibility(UCommonUIVisibilitySubsys
 {
 	if (!IsDesignTime())
 	{
-		const UCommonUIVisibilitySubsystem* CommonInputSubsystem = VisSystem ? VisSystem : UCommonUIVisibilitySubsystem::Get(GetOwningLocalPlayer());
-		if (ensure(CommonInputSubsystem))
-		{
-			bool bVisibleForFeatures = true;
+		bool bVisibleForFeatures = true;
 
-			if (!VisibilityQuery.IsEmpty())
+		if (!VisibilityQuery.IsEmpty())
+		{
+			const UCommonUIVisibilitySubsystem* CommonInputSubsystem = VisSystem ? VisSystem : UCommonUIVisibilitySubsystem::Get(GetOwningLocalPlayer());
+			if (CommonInputSubsystem)
 			{
 				bVisibleForFeatures = VisibilityQuery.Matches(CommonInputSubsystem->GetVisibilityTags());
 			}
+			else
+			{
+				UE_LOG(LogCommonUI, Verbose, TEXT("[%s] -> ULocalPlayer not available, using PlatformTraits instead"), *GetName());
 
-			SetVisibility(bVisibleForFeatures ? VisibleType : HiddenType);
+				// If UCommonUIVisibilitySubsystem is unavailable use the hardware traits from UCommonUISettings
+				const FGameplayTagContainer& HardwareTags = ICommonUIModule::GetSettings().GetPlatformTraits();
+				bVisibleForFeatures = VisibilityQuery.Matches(HardwareTags);
+			}
 		}
+
+		SetVisibility(bVisibleForFeatures ? VisibleType : HiddenType);
 	}
 }
 
