Remove ensure that would always trigger if using CommonUI with default input configs enabled.

diff --git a/Engine/Plugins/Experimental/CommonUI/Source/CommonUI/Private/Input/CommonUIActionRouterBase.cpp b/Engine/Plugins/Experimental/CommonUI/Source/CommonUI/Private/Input/CommonUIActionRouterBase.cpp
--- a/Engine/Plugins/Experimental/CommonUI/Source/CommonUI/Private/Input/CommonUIActionRouterBase.cpp
+++ b/Engine/Plugins/Experimental/CommonUI/Source/CommonUI/Private/Input/CommonUIActionRouterBase.cpp
@@ -1195,7 +1195,6 @@ FActivatableTreeNodePtr UCommonUIActionRouterBase::FindNodeRecursive(const FActi
 
 void UCommonUIActionRouterBase::SetActiveUIInputConfig(const FUIInputConfig& NewConfig)
 {
-	ensureMsgf(ICommonInputModule::GetSettings().GetEnableDefaultInputConfig() == false, TEXT("UIInputConfig is being set but will be overwritten when the activatable tree changes. Disable the default Input Config in Common Input Settings to prevent this."));
 	ApplyUIInputConfig(NewConfig, !ActiveInputConfig.IsSet());
 }
 
