Common Input:
- Change GetCalculatedActionDomain to first check for a meta data on the widget to get the action domain. Else check if the CAW is overriding the action domain
- Add ICommonInputActionDomainMetaData to define the action domain used by a SWidget
- Add description to CVarEnableActionDomainRouting

diff --git a/Engine/Plugins/Experimental/CommonUI/Source/CommonInput/Private/CommonInputActionDomain.cpp b/Engine/Plugins/Experimental/CommonUI/Source/CommonInput/Private/CommonInputActionDomain.cpp
--- a/Engine/Plugins/Experimental/CommonUI/Source/CommonInput/Private/CommonInputActionDomain.cpp
+++ b/Engine/Plugins/Experimental/CommonUI/Source/CommonInput/Private/CommonInputActionDomain.cpp
@@ -42,4 +42,18 @@ bool UCommonInputActionDomain::ShouldBreakEventFlow(bool bDomainHadActiveRoots,
 	}
 
 	return false;
+}
+
+void UCommonInputActionDomainTable::PostLoad()
+{
+	for (UCommonInputActionDomain* ActionDomain : ActionDomains)
+	{
+		if (ActionDomain && ActionDomain->bIsDefaultActionDomain)
+		{
+			DefaultActionDomainCache = ActionDomain;
+			break;
+		}
+	}
+
+	Super::PostLoad();
 }
\ No newline at end of file
diff --git a/Engine/Plugins/Experimental/CommonUI/Source/CommonInput/Public/CommonInputActionDomain.h b/Engine/Plugins/Experimental/CommonUI/Source/CommonInput/Public/CommonInputActionDomain.h
--- a/Engine/Plugins/Experimental/CommonUI/Source/CommonInput/Public/CommonInputActionDomain.h
+++ b/Engine/Plugins/Experimental/CommonUI/Source/CommonInput/Public/CommonInputActionDomain.h
@@ -3,6 +3,7 @@
 #pragma once
 
 #include "Engine/DataAsset.h"
+#include "Types/ISlateMetaData.h"
 
 #include "CommonInputActionDomain.generated.h"
 
@@ -13,6 +14,19 @@ enum class ECommonInputEventFlowBehavior {
 	NeverBlock,
 };
 
+/** Slate meta data to store the owning Action Domain */
+class ICommonInputActionDomainMetaData : public ISlateMetaData
+{
+public:
+	SLATE_METADATA_TYPE(ICommonInputActionDomainMetaData, ISlateMetaData);
+
+	explicit ICommonInputActionDomainMetaData(const TWeakObjectPtr<UCommonInputActionDomain> InActionDomain)
+		: ActionDomain(InActionDomain)
+	{}
+
+	TWeakObjectPtr<UCommonInputActionDomain> ActionDomain;
+};
+
 /**
  * Describes an input-event handling domain. It's InnerBehavior determines how events
  * flow between widgets within the domain and Behavior determines how events will flow to
@@ -24,7 +38,6 @@ class COMMONINPUT_API UCommonInputActionDomain : public UDataAsset
 	GENERATED_BODY()
 
 public:
-
 	// Behavior of an input event between Action Domains, i.e., how an event flows into the next Action Domain
 	UPROPERTY(EditDefaultsOnly, Category = "Default")
 	ECommonInputEventFlowBehavior Behavior = ECommonInputEventFlowBehavior::BlockIfActive;
@@ -34,6 +47,10 @@ public:
 	UPROPERTY(EditDefaultsOnly, Category = "Default")
 	ECommonInputEventFlowBehavior InnerBehavior = ECommonInputEventFlowBehavior::BlockIfHandled;
 
+	// The first Action Domain marked as default will be used when a widget doesn't have any action domain defined in its hierarchie
+	UPROPERTY(EditDefaultsOnly, Category = "Default")
+	bool bIsDefaultActionDomain = false;
+
 	bool ShouldBreakInnerEventFlow(bool bInputEventHandled) const;
 
 	bool ShouldBreakEventFlow(bool bDomainHadActiveRoots, bool bInputEventHandledAtLeastOnce) const;
@@ -51,4 +68,10 @@ public:
 	// Domains will receive events in ascending index order
 	UPROPERTY(EditDefaultsOnly, BlueprintReadOnly, Category = "Default")
 	TArray<UCommonInputActionDomain*> ActionDomains;
+
+	// Cache of the first default action domain found in the table
+	UPROPERTY(Transient)
+	TObjectPtr<UCommonInputActionDomain> DefaultActionDomainCache;
+
+	virtual void PostLoad() override;
 };
\ No newline at end of file
diff --git a/Engine/Plugins/Experimental/CommonUI/Source/CommonUI/Private/CommonActivatableWidget.cpp b/Engine/Plugins/Experimental/CommonUI/Source/CommonUI/Private/CommonActivatableWidget.cpp
--- a/Engine/Plugins/Experimental/CommonUI/Source/CommonUI/Private/CommonActivatableWidget.cpp
+++ b/Engine/Plugins/Experimental/CommonUI/Source/CommonUI/Private/CommonActivatableWidget.cpp
@@ -6,6 +6,7 @@
 #include "Input/CommonUIInputTypes.h"
 #include "Input/UIActionRouterTypes.h"
 #include "ICommonInputModule.h"
+#include "Slate/SObjectWidget.h"
 
 UCommonActivatableWidget::FActivatableWidgetRebuildEvent UCommonActivatableWidget::OnRebuilding;
 
@@ -149,20 +150,42 @@ void UCommonActivatableWidget::ClearActiveHoldInputs()
 
 TObjectPtr<UCommonInputActionDomain> UCommonActivatableWidget::GetCalculatedActionDomain()
 {
-	if (!bInheritActionDomain)
+	if (CalculatedActionDomainCache.IsValid())
 	{
-		return ActionDomain.Get();
+		return CalculatedActionDomainCache.Get();
 	}
 
-	const UCommonActivatableWidget* CurrentWidget = this;
-	while (CurrentWidget && CurrentWidget->bInheritActionDomain)
+	if (bOverrideActionDomain)
 	{
-		CurrentWidget = UCommonUIActionRouterBase::FindOwningActivatable(CurrentWidget->GetCachedWidget(), GetOwningLocalPlayer());
+		CalculatedActionDomainCache = ActionDomain;
+		return ActionDomain.Get();
 	}
 
-	if (CurrentWidget && !CurrentWidget->bInheritActionDomain)
+	const FName SObjectWidgetName = TEXT("SObjectWidget");
+	const ULocalPlayer* OwningLocalPlayer = GetOwningLocalPlayer();
+	TSharedPtr<SWidget> CurrentWidget = GetCachedWidget();
+	while (CurrentWidget)
 	{
-		return CurrentWidget->ActionDomain.Get();
+		CurrentWidget = CurrentWidget->GetParentWidget();
+		if (CurrentWidget && CurrentWidget->GetType().IsEqual(SObjectWidgetName))
+		{
+			const TSharedPtr<ICommonInputActionDomainMetaData> Metadata = CurrentWidget->GetMetaData<ICommonInputActionDomainMetaData>();
+			if (Metadata.IsValid())
+			{
+				CalculatedActionDomainCache = Metadata->ActionDomain.Get();
+				return CalculatedActionDomainCache.Get();
+			}
+
+			if (UCommonActivatableWidget* CurrentActivatable = Cast<UCommonActivatableWidget>(StaticCastSharedPtr<SObjectWidget>(CurrentWidget)->GetWidgetObject()))
+			{
+				if (CurrentActivatable->bOverrideActionDomain)
+				{
+					UCommonInputActionDomain* CurrentActionDomain = CurrentActivatable->GetOwningLocalPlayer() == OwningLocalPlayer ? CurrentActivatable->ActionDomain.Get() : nullptr;
+					CalculatedActionDomainCache = CurrentActionDomain;
+					return CalculatedActionDomainCache.Get();
+				}
+			}
+		}
 	}
 
 	return nullptr;
diff --git a/Engine/Plugins/Experimental/CommonUI/Source/CommonUI/Private/Input/CommonUIActionRouterBase.cpp b/Engine/Plugins/Experimental/CommonUI/Source/CommonUI/Private/Input/CommonUIActionRouterBase.cpp
--- a/Engine/Plugins/Experimental/CommonUI/Source/CommonUI/Private/Input/CommonUIActionRouterBase.cpp
+++ b/Engine/Plugins/Experimental/CommonUI/Source/CommonUI/Private/Input/CommonUIActionRouterBase.cpp
@@ -35,7 +35,7 @@ bool bEnableActionDomainRouting = false;
 static const FAutoConsoleVariableRef CVarEnableActionDomainRouting(
 	TEXT("CommonUI.EnableActionDomainRouting"),
 	bEnableActionDomainRouting,
-	TEXT(""));
+	TEXT("Enables the routing of UI inputs based on the action domain instead of the active root node."));
 
 //@todo DanH: TEMP LOCATION
 FGlobalUITags FGlobalUITags::GUITags;
@@ -663,7 +663,15 @@ void UCommonUIActionRouterBase::AddToActionDomain(FActivatableTreeRootRef RootNo
 	}
 	else
 	{
-		UE_LOG(LogUIActionRouter, Error, TEXT("ActionDomain of CommonActivatableWidget could not be resolved. Widget: %s"), *RootNode->GetWidget()->GetName());
+		UCommonInputSubsystem& CommonInputSubsystem = GetInputSubsystem();
+		UCommonInputActionDomainTable* ActionDomainTable = CommonInputSubsystem.GetActionDomainTable();
+		if (!ActionDomainTable->DefaultActionDomainCache)
+		{
+			UE_LOG(LogUIActionRouter, Error, TEXT("ActionDomain of CommonActivatableWidget could not be resolved. Widget: %s"), *RootNode->GetWidget()->GetName());
+			return;
+		}
+
+		ActionDomainRootNodes.FindOrAdd(ActionDomainTable->DefaultActionDomainCache).Add(RootNode);
 	}
 }
 
@@ -671,18 +679,17 @@ void UCommonUIActionRouterBase::RemoveFromActionDomain(FActivatableTreeRootRef R
 {
 	UCommonInputActionDomain* ActionDomain = RootNode->GetWidget()->GetCalculatedActionDomain();
 
-	if (ActionDomain)
+	if (!ActionDomain)
 	{
-		FActionDomainSortedRootList* ActionDomainRootList = ActionDomainRootNodes.Find(ActionDomain);
-
-		if (ensure(ActionDomainRootList))
-		{
-			verify(ActionDomainRootList->Remove(RootNode) != INDEX_NONE);
-		}
+		UCommonInputSubsystem& CommonInputSubsystem = GetInputSubsystem();
+		UCommonInputActionDomainTable* ActionDomainTable = CommonInputSubsystem.GetActionDomainTable();
+		ActionDomain = ActionDomainTable ? ActionDomainTable->DefaultActionDomainCache : nullptr;
 	}
-	else
+
+	FActionDomainSortedRootList* ActionDomainRootList = ActionDomainRootNodes.Find(ActionDomain);
+	if (ensure(ActionDomainRootList))
 	{
-		UE_LOG(LogUIActionRouter, Error, TEXT("ActionDomain of CommonActivatableWidget could not be resolved. Widget: %s"), *RootNode->GetWidget()->GetName());
+		verify(ActionDomainRootList->Remove(RootNode) != INDEX_NONE);
 	}
 }
 
diff --git a/Engine/Plugins/Experimental/CommonUI/Source/CommonUI/Public/CommonActivatableWidget.h b/Engine/Plugins/Experimental/CommonUI/Source/CommonUI/Public/CommonActivatableWidget.h
--- a/Engine/Plugins/Experimental/CommonUI/Source/CommonUI/Public/CommonActivatableWidget.h
+++ b/Engine/Plugins/Experimental/CommonUI/Source/CommonUI/Public/CommonActivatableWidget.h
@@ -186,12 +186,12 @@ protected:
 	bool bAutoRestoreFocus = false;
 
 	UPROPERTY(EditAnywhere, BlueprintReadOnly, Category = "Input | ActionDomain", meta = (InlineEditConditionToggle))
-	bool bInheritActionDomain = true;
+	bool bOverrideActionDomain = false;
 
 	/**
-	 * Disable to inherit from owning CommonActivatableWidget.
+	 * Enable to override the inheritd ActionDomain from owning CommonActivatableWidget.
 	 */
-	UPROPERTY(EditAnywhere, BlueprintReadOnly, Category = "Input | ActionDomain", meta = (EditCondition = "!bInheritActionDomain", DisplayName = "Override ActionDomain"))
+	UPROPERTY(EditAnywhere, BlueprintReadOnly, Category = "Input | ActionDomain", meta = (EditCondition = "bOverrideActionDomain", DisplayName = "Override ActionDomain"))
 	TSoftObjectPtr<UCommonInputActionDomain> ActionDomain;
 
 private:
@@ -233,6 +233,8 @@ private:
 	mutable FSimpleMulticastDelegate OnSlateReleasedEvent;
 	mutable FSimpleMulticastDelegate OnRequestRefreshFocusEvent;
 
+	TSoftObjectPtr<UCommonInputActionDomain> CalculatedActionDomainCache;
+
 protected:
 	
 	UPROPERTY(EditAnywhere, Category = Activation, meta = (InlineEditConditionToggle = "ActivatedVisibility"))
