Common Input:
- Made action domain cache optional to denote if a widget tried to get an action domain but couldn't versus one that simply didn't try yet
- Evaluate the current widget when calculating the action domain as it might have a meta data
- Fixed crash when stopping in PIE without a full setup when action domain are enabled

diff --git a/Engine/Plugins/Experimental/CommonUI/Source/CommonInput/Private/CommonInputActionDomain.cpp b/Engine/Plugins/Experimental/CommonUI/Source/CommonInput/Private/CommonInputActionDomain.cpp
--- a/Engine/Plugins/Experimental/CommonUI/Source/CommonInput/Private/CommonInputActionDomain.cpp
+++ b/Engine/Plugins/Experimental/CommonUI/Source/CommonInput/Private/CommonInputActionDomain.cpp
@@ -2,6 +2,8 @@
 
 #include "CommonInputActionDomain.h"
 
+DEFINE_LOG_CATEGORY(LogUIActionDomain);
+
 bool UCommonInputActionDomain::ShouldBreakInnerEventFlow(bool bInputEventHandled) const
 {
 	switch (InnerBehavior)
@@ -55,5 +57,10 @@ void UCommonInputActionDomainTable::PostLoad()
 		}
 	}
 
+	if (DefaultActionDomainCache.IsNull())
+	{
+		UE_LOG(LogUIActionDomain, Error, TEXT("No default action domain defined"));
+	}
+
 	Super::PostLoad();
 }
\ No newline at end of file
diff --git a/Engine/Plugins/Experimental/CommonUI/Source/CommonInput/Public/CommonInputActionDomain.h b/Engine/Plugins/Experimental/CommonUI/Source/CommonInput/Public/CommonInputActionDomain.h
--- a/Engine/Plugins/Experimental/CommonUI/Source/CommonInput/Public/CommonInputActionDomain.h
+++ b/Engine/Plugins/Experimental/CommonUI/Source/CommonInput/Public/CommonInputActionDomain.h
@@ -7,6 +7,8 @@
 
 #include "CommonInputActionDomain.generated.h"
 
+DECLARE_LOG_CATEGORY_EXTERN(LogUIActionDomain, Log, All);
+
 UENUM()
 enum class ECommonInputEventFlowBehavior {
 	BlockIfActive,
diff --git a/Engine/Plugins/Experimental/CommonUI/Source/CommonUI/Private/CommonActivatableWidget.cpp b/Engine/Plugins/Experimental/CommonUI/Source/CommonUI/Private/CommonActivatableWidget.cpp
--- a/Engine/Plugins/Experimental/CommonUI/Source/CommonUI/Private/CommonActivatableWidget.cpp
+++ b/Engine/Plugins/Experimental/CommonUI/Source/CommonUI/Private/CommonActivatableWidget.cpp
@@ -150,15 +150,9 @@ void UCommonActivatableWidget::ClearActiveHoldInputs()
 
 TObjectPtr<UCommonInputActionDomain> UCommonActivatableWidget::GetCalculatedActionDomain()
 {
-	if (CalculatedActionDomainCache.IsValid())
+	if (CalculatedActionDomainCache.IsSet())
 	{
-		return CalculatedActionDomainCache.Get();
-	}
-
-	if (bOverrideActionDomain)
-	{
-		CalculatedActionDomainCache = ActionDomain;
-		return ActionDomain.Get();
+		return CalculatedActionDomainCache.GetValue().Get();
 	}
 
 	const FName SObjectWidgetName = TEXT("SObjectWidget");
@@ -166,28 +160,31 @@ TObjectPtr<UCommonInputActionDomain> UCommonActivatableWidget::GetCalculatedActi
 	TSharedPtr<SWidget> CurrentWidget = GetCachedWidget();
 	while (CurrentWidget)
 	{
-		CurrentWidget = CurrentWidget->GetParentWidget();
-		if (CurrentWidget && CurrentWidget->GetType().IsEqual(SObjectWidgetName))
+		if (CurrentWidget->GetType().IsEqual(SObjectWidgetName))
 		{
 			const TSharedPtr<ICommonInputActionDomainMetaData> Metadata = CurrentWidget->GetMetaData<ICommonInputActionDomainMetaData>();
 			if (Metadata.IsValid())
 			{
-				CalculatedActionDomainCache = Metadata->ActionDomain.Get();
-				return CalculatedActionDomainCache.Get();
+				UCommonInputActionDomain* ActionDomain = Metadata->ActionDomain.Get();
+				CalculatedActionDomainCache = ActionDomain;
+				return ActionDomain;
 			}
 
 			if (UCommonActivatableWidget* CurrentActivatable = Cast<UCommonActivatableWidget>(StaticCastSharedPtr<SObjectWidget>(CurrentWidget)->GetWidgetObject()))
 			{
 				if (CurrentActivatable->bOverrideActionDomain)
 				{
-					UCommonInputActionDomain* CurrentActionDomain = CurrentActivatable->GetOwningLocalPlayer() == OwningLocalPlayer ? CurrentActivatable->ActionDomain.Get() : nullptr;
+					UCommonInputActionDomain* CurrentActionDomain = CurrentActivatable->GetOwningLocalPlayer() == OwningLocalPlayer ? CurrentActivatable->ActionDomainOverride.Get() : nullptr;
 					CalculatedActionDomainCache = CurrentActionDomain;
-					return CalculatedActionDomainCache.Get();
+					return CurrentActionDomain;
 				}
 			}
 		}
+
+		CurrentWidget = CurrentWidget->GetParentWidget();
 	}
 
+	CalculatedActionDomainCache = nullptr;
 	return nullptr;
 }
 
diff --git a/Engine/Plugins/Experimental/CommonUI/Source/CommonUI/Private/Input/CommonUIActionRouterBase.cpp b/Engine/Plugins/Experimental/CommonUI/Source/CommonUI/Private/Input/CommonUIActionRouterBase.cpp
--- a/Engine/Plugins/Experimental/CommonUI/Source/CommonUI/Private/Input/CommonUIActionRouterBase.cpp
+++ b/Engine/Plugins/Experimental/CommonUI/Source/CommonUI/Private/Input/CommonUIActionRouterBase.cpp
@@ -31,7 +31,7 @@ static const FAutoConsoleVariableRef CVarAlwaysShowCursor(
 	bAlwaysShowCursor,
 	TEXT(""));
 
-bool bEnableActionDomainRouting = false;
+bool bEnableActionDomainRouting = true;
 static const FAutoConsoleVariableRef CVarEnableActionDomainRouting(
 	TEXT("CommonUI.EnableActionDomainRouting"),
 	bEnableActionDomainRouting,
@@ -681,13 +681,15 @@ void UCommonUIActionRouterBase::RemoveFromActionDomain(FActivatableTreeRootRef R
 
 	if (!ActionDomain)
 	{
-		UCommonInputSubsystem& CommonInputSubsystem = GetInputSubsystem();
-		UCommonInputActionDomainTable* ActionDomainTable = CommonInputSubsystem.GetActionDomainTable();
-		ActionDomain = ActionDomainTable ? ActionDomainTable->DefaultActionDomainCache : nullptr;
+		if (UCommonInputSubsystem* InputSubsytem = GetLocalPlayerChecked()->GetSubsystem<UCommonInputSubsystem>())
+		{
+			UCommonInputActionDomainTable* ActionDomainTable = InputSubsytem->GetActionDomainTable();
+			ActionDomain = ActionDomainTable ? ActionDomainTable->DefaultActionDomainCache : nullptr;
+		}
 	}
 
 	FActionDomainSortedRootList* ActionDomainRootList = ActionDomainRootNodes.Find(ActionDomain);
-	if (ensure(ActionDomainRootList))
+	if (ActionDomainRootList)
 	{
 		verify(ActionDomainRootList->Remove(RootNode) != INDEX_NONE);
 	}
diff --git a/Engine/Plugins/Experimental/CommonUI/Source/CommonUI/Public/CommonActivatableWidget.h b/Engine/Plugins/Experimental/CommonUI/Source/CommonUI/Public/CommonActivatableWidget.h
--- a/Engine/Plugins/Experimental/CommonUI/Source/CommonUI/Public/CommonActivatableWidget.h
+++ b/Engine/Plugins/Experimental/CommonUI/Source/CommonUI/Public/CommonActivatableWidget.h
@@ -189,10 +189,10 @@ protected:
 	bool bOverrideActionDomain = false;
 
 	/**
-	 * Enable to override the inheritd ActionDomain from owning CommonActivatableWidget.
+	 * Enable to override the inherited ActionDomain from owning CommonActivatableWidget.
 	 */
-	UPROPERTY(EditAnywhere, BlueprintReadOnly, Category = "Input | ActionDomain", meta = (EditCondition = "bOverrideActionDomain", DisplayName = "Override ActionDomain"))
-	TSoftObjectPtr<UCommonInputActionDomain> ActionDomain;
+	UPROPERTY(EditAnywhere, BlueprintReadOnly, Category = "Input | ActionDomain", meta = (EditCondition = "bOverrideActionDomain"))
+	TSoftObjectPtr<UCommonInputActionDomain> ActionDomainOverride;
 
 private:
 	/** See BindVisibilityToMultipleActivations */
@@ -233,7 +233,7 @@ private:
 	mutable FSimpleMulticastDelegate OnSlateReleasedEvent;
 	mutable FSimpleMulticastDelegate OnRequestRefreshFocusEvent;
 
-	TSoftObjectPtr<UCommonInputActionDomain> CalculatedActionDomainCache;
+	TOptional<TSoftObjectPtr<UCommonInputActionDomain>> CalculatedActionDomainCache;
 
 protected:
 	
