Common Input:
- Remove default action domain

diff --git a/Engine/Plugins/Experimental/CommonUI/Source/CommonInput/Private/CommonInputActionDomain.cpp b/Engine/Plugins/Experimental/CommonUI/Source/CommonInput/Private/CommonInputActionDomain.cpp
--- a/Engine/Plugins/Experimental/CommonUI/Source/CommonInput/Private/CommonInputActionDomain.cpp
+++ b/Engine/Plugins/Experimental/CommonUI/Source/CommonInput/Private/CommonInputActionDomain.cpp
@@ -45,22 +45,3 @@ bool UCommonInputActionDomain::ShouldBreakEventFlow(bool bDomainHadActiveRoots,
 
 	return false;
 }
-
-void UCommonInputActionDomainTable::PostLoad()
-{
-	for (UCommonInputActionDomain* ActionDomain : ActionDomains)
-	{
-		if (ActionDomain && ActionDomain->bIsDefaultActionDomain)
-		{
-			DefaultActionDomainCache = ActionDomain;
-			break;
-		}
-	}
-
-	if (DefaultActionDomainCache.IsNull())
-	{
-		UE_LOG(LogUIActionDomain, Error, TEXT("No default action domain defined"));
-	}
-
-	Super::PostLoad();
-}
\ No newline at end of file
diff --git a/Engine/Plugins/Experimental/CommonUI/Source/CommonInput/Public/CommonInputActionDomain.h b/Engine/Plugins/Experimental/CommonUI/Source/CommonInput/Public/CommonInputActionDomain.h
--- a/Engine/Plugins/Experimental/CommonUI/Source/CommonInput/Public/CommonInputActionDomain.h
+++ b/Engine/Plugins/Experimental/CommonUI/Source/CommonInput/Public/CommonInputActionDomain.h
@@ -49,10 +49,6 @@ public:
 	UPROPERTY(EditDefaultsOnly, Category = "Default")
 	ECommonInputEventFlowBehavior InnerBehavior = ECommonInputEventFlowBehavior::BlockIfHandled;
 
-	// The first Action Domain marked as default will be used when a widget doesn't have any action domain defined in its hierarchie
-	UPROPERTY(EditDefaultsOnly, Category = "Default")
-	bool bIsDefaultActionDomain = false;
-
 	bool ShouldBreakInnerEventFlow(bool bInputEventHandled) const;
 
 	bool ShouldBreakEventFlow(bool bDomainHadActiveRoots, bool bInputEventHandledAtLeastOnce) const;
@@ -70,10 +66,4 @@ public:
 	// Domains will receive events in ascending index order
 	UPROPERTY(EditDefaultsOnly, BlueprintReadOnly, Category = "Default")
 	TArray<UCommonInputActionDomain*> ActionDomains;
-
-	// Cache of the first default action domain found in the table
-	UPROPERTY(Transient)
-	TObjectPtr<UCommonInputActionDomain> DefaultActionDomainCache;
-
-	virtual void PostLoad() override;
 };
\ No newline at end of file
