CIS fix

diff --git a/Engine/Plugins/Experimental/CommonUI/Source/CommonUI/Private/Input/CommonUIActionRouterBase.cpp b/Engine/Plugins/Experimental/CommonUI/Source/CommonUI/Private/Input/CommonUIActionRouterBase.cpp
--- a/Engine/Plugins/Experimental/CommonUI/Source/CommonUI/Private/Input/CommonUIActionRouterBase.cpp
+++ b/Engine/Plugins/Experimental/CommonUI/Source/CommonUI/Private/Input/CommonUIActionRouterBase.cpp
@@ -661,33 +661,11 @@ void UCommonUIActionRouterBase::AddToActionDomain(FActivatableTreeRootRef RootNo
 	{
 		ActionDomainRootNodes.FindOrAdd(ActionDomain).Add(RootNode);
 	}
-	else
-	{
-		UCommonInputSubsystem& CommonInputSubsystem = GetInputSubsystem();
-		UCommonInputActionDomainTable* ActionDomainTable = CommonInputSubsystem.GetActionDomainTable();
-		if (!ActionDomainTable->DefaultActionDomainCache)
-		{
-			UE_LOG(LogUIActionRouter, Error, TEXT("ActionDomain of CommonActivatableWidget could not be resolved. Widget: %s"), *RootNode->GetWidget()->GetName());
-			return;
-		}
-
-		ActionDomainRootNodes.FindOrAdd(ActionDomainTable->DefaultActionDomainCache).Add(RootNode);
-	}
 }
 
 void UCommonUIActionRouterBase::RemoveFromActionDomain(FActivatableTreeRootRef RootNode)
 {
 	UCommonInputActionDomain* ActionDomain = RootNode->GetWidget()->GetCalculatedActionDomain();
-
-	if (!ActionDomain)
-	{
-		if (UCommonInputSubsystem* InputSubsytem = GetLocalPlayerChecked()->GetSubsystem<UCommonInputSubsystem>())
-		{
-			UCommonInputActionDomainTable* ActionDomainTable = InputSubsytem->GetActionDomainTable();
-			ActionDomain = ActionDomainTable ? ActionDomainTable->DefaultActionDomainCache : nullptr;
-		}
-	}
-
 	FActionDomainSortedRootList* ActionDomainRootList = ActionDomainRootNodes.Find(ActionDomain);
 	if (ActionDomainRootList)
 	{
