CommonUI - The CommonUI Action router attempts to be 'smart' and allow PIE's "Escape" command to bypass the UI Action router so that it can actually stop PIE, the problem is that it makes a huge assumption that the user is actually using Escape with no modifier instead of checking the command to see what they're actually using and instead simply bypass UI Action Router if you ever press the PIE stop command - no matter what the user has bound it to.

diff --git a/Engine/Plugins/Experimental/CommonUI/Source/CommonUI/Private/Input/CommonUIActionRouterBase.cpp b/Engine/Plugins/Experimental/CommonUI/Source/CommonUI/Private/Input/CommonUIActionRouterBase.cpp
--- a/Engine/Plugins/Experimental/CommonUI/Source/CommonUI/Private/Input/CommonUIActionRouterBase.cpp
+++ b/Engine/Plugins/Experimental/CommonUI/Source/CommonUI/Private/Input/CommonUIActionRouterBase.cpp
@@ -318,22 +318,20 @@ TSharedRef<FCommonAnalogCursor> UCommonUIActionRouterBase::MakeAnalogCursor() co
 ERouteUIInputResult UCommonUIActionRouterBase::ProcessInput(FKey Key, EInputEvent InputEvent) const
 {
 #if WITH_EDITOR
-	// In PIE, let unmodified escape through (people expect it to close PIE)
-	if (GIsPlayInEditorWorld && InputEvent == IE_Pressed && Key == EKeys::Escape)
+	// In PIE, check if the user is attempting to press the StopPlaySession command chord.
+	if (GIsPlayInEditorWorld && InputEvent == IE_Pressed)
 	{
+		//TODO This could be more generic to be a list of command in commands to allow to be ignored by the UI action router.
 		TSharedPtr<FUICommandInfo> StopCommand = FInputBindingManager::Get().FindCommandInContext("PlayWorld", "StopPlaySession");
 		if (ensure(StopCommand))
 		{
-			// If the user has rebound the 'Stop' to something other than vanilla Escape, then don't do this,
-			// they're clearly trying to avoid collisions with the UI Actions.
-			const TSharedRef<const FInputChord> PrimaryCord = StopCommand->GetActiveChord(EMultipleKeyBindingIndex::Primary);
-			if (PrimaryCord->Key == EKeys::Escape && PrimaryCord->bAlt == false && PrimaryCord->bCmd == false && PrimaryCord->bCtrl == false && PrimaryCord->bShift == false)
+			const FModifierKeysState ModifierKeys = FSlateApplication::Get().GetModifierKeys();
+			const FInputChord CheckChord( Key, EModifierKey::FromBools(ModifierKeys.IsControlDown(), ModifierKeys.IsAltDown(), ModifierKeys.IsShiftDown(), ModifierKeys.IsCommandDown()) );
+
+			// If the stop command matches the incoming key chord, let it execute.
+			if (StopCommand->HasActiveChord(CheckChord))
 			{
-				FModifierKeysState ModifierKeys = FSlateApplication::Get().GetModifierKeys();
-				if (!ModifierKeys.IsAltDown() && !ModifierKeys.IsCommandDown() && !ModifierKeys.IsControlDown() && !ModifierKeys.IsShiftDown())
-				{
-					return ERouteUIInputResult::Unhandled;
-				}
+				return ERouteUIInputResult::Unhandled;
 			}
 		}
 	}
