Fix Ensure quitting PIE session in Lyra using Quit Game.

diff --git a/Engine/Plugins/Experimental/CommonUI/Source/CommonUI/Private/Widgets/CommonActivatableWidgetContainer.cpp b/Engine/Plugins/Experimental/CommonUI/Source/CommonUI/Private/Widgets/CommonActivatableWidgetContainer.cpp
--- a/Engine/Plugins/Experimental/CommonUI/Source/CommonUI/Private/Widgets/CommonActivatableWidgetContainer.cpp
+++ b/Engine/Plugins/Experimental/CommonUI/Source/CommonUI/Private/Widgets/CommonActivatableWidgetContainer.cpp
@@ -124,6 +124,7 @@ void UCommonActivatableWidgetContainerBase::ReleaseSlateResources(bool bReleaseC
 	MyOverlay.Reset();
 	MyInputGuard.Reset();
 	MySwitcher.Reset();
+	ReleasedWidgets.Empty();
 
 	GeneratedWidgetsPool.ReleaseAllSlateResources();
 }
