Allow tabs to be inserted to tab list.

UCommonTabListWidgetBase:
- Added TabIndex parameter to RegisterTab (default value -1 which means tab will be added to end).
- If new tab is being inserted before end of list, rebuild tab list with tabs in correct order.
- Added option to defer rebuilding till next tick (so we only rebuild once if multiple tabs are inserted).

diff --git a/Engine/Plugins/Experimental/CommonUI/Source/CommonUI/Private/CommonTabListWidgetBase.cpp b/Engine/Plugins/Experimental/CommonUI/Source/CommonUI/Private/CommonTabListWidgetBase.cpp
--- a/Engine/Plugins/Experimental/CommonUI/Source/CommonUI/Private/CommonTabListWidgetBase.cpp
+++ b/Engine/Plugins/Experimental/CommonUI/Source/CommonUI/Private/CommonTabListWidgetBase.cpp
@@ -3,21 +3,25 @@
 #include "CommonTabListWidgetBase.h"
 #include "CommonUIPrivatePCH.h"
 
+#include "Blueprint/WidgetLayoutLibrary.h"
+#include "CommonAnimatedSwitcher.h"
 #include "CommonUISubsystemBase.h"
+#include "CommonUIUtils.h"
+#include "Containers/Ticker.h"
 #include "Groups/CommonButtonGroupBase.h"
-#include "CommonAnimatedSwitcher.h"
-#include "Blueprint/WidgetLayoutLibrary.h"
 #include "ICommonUIModule.h"
-#include "CommonUIUtils.h"
 #include "Input/CommonUIInputTypes.h"
 
 UCommonTabListWidgetBase::UCommonTabListWidgetBase(const FObjectInitializer& ObjectInitializer)
 	: Super(ObjectInitializer)
 	, bAutoListenForInput(false)
+	, bDeferRebuildingTabList(false)
 	, TabButtonGroup(nullptr)
 	, bIsListeningForInput(false)
 	, RegisteredTabsByID()
 	, ActiveTabID(NAME_None)
+	, bIsRebuildingList(false)
+	, bPendingRebuild(false)
 {
 }
  
@@ -36,23 +40,30 @@ UCommonAnimatedSwitcher* UCommonTabListWidgetBase::GetLinkedSwitcher() const
 	return LinkedSwitcher.Get();
 }
 
-bool UCommonTabListWidgetBase::RegisterTab(FName TabNameID, TSubclassOf<UCommonButtonBase> ButtonWidgetType, UWidget* ContentWidget)
+bool UCommonTabListWidgetBase::RegisterTab(FName TabNameID, TSubclassOf<UCommonButtonBase> ButtonWidgetType, UWidget* ContentWidget, const int32 TabIndex /*= INDEX_NONE*/)
 {
-	bool AreParametersValid = true;
+	bool bAreParametersValid = true;
 
 	// Early out on redundant tab registration.
 	if (!ensure(!RegisteredTabsByID.Contains(TabNameID)))
 	{
-		AreParametersValid = false;
+		bAreParametersValid = false;
 	}
 
 	// Early out on invalid tab button type.
 	if (!ensure(ButtonWidgetType))
 	{
-		AreParametersValid = false;
+		bAreParametersValid = false;
+	}
+	
+	// NOTE: Adding the button to the group may change it's selection, which raises an event we listen to,
+	// which can only properly be handled if we already know that this button is associated with a registered tab.
+	if (!ensure(TabButtonGroup))
+	{
+		bAreParametersValid = false;
 	}
 
-	if (!AreParametersValid)
+	if (!bAreParametersValid)
 	{
 		return false;
 	}
@@ -63,9 +74,27 @@ bool UCommonTabListWidgetBase::RegisterTab(FName TabNameID, TSubclassOf<UCommonB
 		return false;
 	}
 
+	const int32 NumRegisteredTabs = RegisteredTabsByID.Num();
+	const int32 NewTabIndex = (TabIndex == INDEX_NONE) ? NumRegisteredTabs : FMath::Clamp(TabIndex, 0, NumRegisteredTabs);
+
+	// If the new tab is being inserted before the end of the list, we need to rebuild the tab list.
+	const bool bRequiresRebuild = (NewTabIndex < NumRegisteredTabs);
+
+	if (bRequiresRebuild)
+	{
+		for (TPair<FName, FCommonRegisteredTabInfo>& Pair : RegisteredTabsByID)
+		{
+			if (NewTabIndex <= Pair.Value.TabIndex)
+			{
+				// Increment this tab's index as we are inserting the new tab before it.
+				Pair.Value.TabIndex++;
+			}
+		}
+	}
+
 	// Tab book-keeping.
 	FCommonRegisteredTabInfo NewTabInfo;
-	NewTabInfo.TabIndex = RegisteredTabsByID.Num();
+	NewTabInfo.TabIndex = NewTabIndex;
 	NewTabInfo.TabButton = NewTabButton;
 	NewTabInfo.ContentInstance = ContentWidget;
 	RegisteredTabsByID.Add(TabNameID, NewTabInfo);
@@ -73,21 +102,32 @@ bool UCommonTabListWidgetBase::RegisterTab(FName TabNameID, TSubclassOf<UCommonB
 	// Enforce the "contract" that tab buttons require - single-selectability, but not toggleability.
 	NewTabButton->SetIsSelectable(true);
 	NewTabButton->SetIsToggleable(false);
-	// NOTE: Adding the button to the group may change it's selection, which raises an event we listen to,
-	// which can only properly be handled if we already know that this button is associated with a registered tab.
-	if (ensure(TabButtonGroup != nullptr))
-	{
-		TabButtonGroup->AddWidget(NewTabButton);
-	}
 
-	// Callbacks.
+	TabButtonGroup->AddWidget(NewTabButton);
 	HandleTabCreation(TabNameID, NewTabInfo.TabButton);
+	
 	OnTabButtonCreation.Broadcast(TabNameID, NewTabInfo.TabButton);
+	
+	if (bRequiresRebuild)
+	{
+		if (bDeferRebuildingTabList)
+		{
+			if (!bPendingRebuild)
+			{
+				bPendingRebuild = true;
+
+				FTSTicker::GetCoreTicker().AddTicker(FTickerDelegate::CreateUObject(this, &UCommonTabListWidgetBase::DeferredRebuildTabList));
+			}
+		}
+		else
+		{
+			RebuildTabList();
+		}
+	}
 
 	return true;
 }
 
-
 bool UCommonTabListWidgetBase::RemoveTab(FName TabNameID)
 {
 	FCommonRegisteredTabInfo* const TabInfo = RegisteredTabsByID.Find(TabNameID);
@@ -161,6 +201,11 @@ void UCommonTabListWidgetBase::UpdateBindings()
 	}
 }
 
+bool UCommonTabListWidgetBase::IsRebuildingList() const
+{
+	return bIsRebuildingList;
+}
+
 bool UCommonTabListWidgetBase::SelectTabByID(FName TabNameID, bool bSuppressClickFeedback)
 {
 	for (auto& TabPair : RegisteredTabsByID)
@@ -346,6 +391,11 @@ void UCommonTabListWidgetBase::HandleTabRemoval_Implementation(FName TabNameID,
 {
 }
 
+const TMap<FName, FCommonRegisteredTabInfo>& UCommonTabListWidgetBase::GetRegisteredTabsByID() const
+{
+	return RegisteredTabsByID;
+}
+
 void UCommonTabListWidgetBase::HandleTabButtonSelected(UCommonButtonBase* SelectedTabButton, int32 ButtonIndex)
 {
 	for (auto& TabPair : RegisteredTabsByID)
@@ -395,4 +445,47 @@ void UCommonTabListWidgetBase::HandlePreviousTabAction()
 	{
 		TabButtonGroup->SelectPreviousButton();
 	}
-}
\ No newline at end of file
+}
+
+bool UCommonTabListWidgetBase::DeferredRebuildTabList(float DeltaTime)
+{
+	bPendingRebuild = false;
+	RebuildTabList();
+	return false;
+}
+
+void UCommonTabListWidgetBase::RebuildTabList()
+{
+	bIsRebuildingList = true;
+
+	// Copy the registered tabs (as we are about to clear them) and sort by TabIndex.
+	TMap<FName, FCommonRegisteredTabInfo> SortedRegisteredTabsByID = RegisteredTabsByID;
+	SortedRegisteredTabsByID.ValueSort([](const FCommonRegisteredTabInfo& TabInfoA, const FCommonRegisteredTabInfo& TabInfoB)
+		{
+			return (TabInfoA.TabIndex < TabInfoB.TabIndex);
+		});
+
+	// Keep track of the current ActiveTabID so we can restore it after the list is rebuilt.
+	const FName CurrentActiveTabID = ActiveTabID;
+
+	// Disable selection required temporarily so we can deselect everything, rebuild the list, then select the tab we want.
+	TabButtonGroup->SetSelectionRequired(false);
+	TabButtonGroup->DeselectAll();
+	RemoveAllTabs();
+
+	RegisteredTabsByID = SortedRegisteredTabsByID;
+
+	for (TPair<FName, FCommonRegisteredTabInfo>& Pair : RegisteredTabsByID)
+	{
+		TabButtonGroup->AddWidget(Pair.Value.TabButton);
+		HandleTabCreation(Pair.Key, Pair.Value.TabButton);
+	}
+
+	bIsRebuildingList = false;
+
+	constexpr bool bSuppressClickFeedback = true;
+	SelectTabByID(CurrentActiveTabID, bSuppressClickFeedback);
+
+	TabButtonGroup->SetSelectionRequired(true);
+	OnTabListRebuilt.Broadcast();
+}
diff --git a/Engine/Plugins/Experimental/CommonUI/Source/CommonUI/Public/CommonTabListWidgetBase.h b/Engine/Plugins/Experimental/CommonUI/Source/CommonUI/Public/CommonTabListWidgetBase.h
--- a/Engine/Plugins/Experimental/CommonUI/Source/CommonUI/Public/CommonTabListWidgetBase.h
+++ b/Engine/Plugins/Experimental/CommonUI/Source/CommonUI/Public/CommonTabListWidgetBase.h
@@ -63,6 +63,13 @@ public:
 	/** Broadcasts when a new tab is created. */
 	UPROPERTY(BlueprintAssignable, Category = TabList)
 	FOnTabButtonRemoval OnTabButtonRemoval;
+	
+	/** Delegate broadcast when the tab list has been rebuilt (after a new tab has been inserted rather than added to the end). */
+	DECLARE_DYNAMIC_MULTICAST_DELEGATE(FOnTabListRebuilt);
+	
+	/** Broadcasts when the tab list has been rebuilt (after a new tab has been inserted rather than added to the end). */
+	UPROPERTY(BlueprintAssignable, Category = TabList)
+	FOnTabListRebuilt OnTabListRebuilt;
 
 	/** @return The currently active (selected) tab */
 	UFUNCTION(BlueprintCallable, Category = TabList)
@@ -84,10 +91,11 @@ public:
 	 * @param TabID The name ID used to keep track of this tab. Attempts to register a tab under a duplicate ID will fail.
 	 * @param ButtonWidgetType The widget type to create for this tab
 	 * @param ContentWidget The widget to associate with the registered tab
+	 * @param TabIndex Determines where in the tab list to insert the new tab (-1 means tab will be added to end of the list)
 	 * @return True if the new tab registered successfully and there were no name ID conflicts
 	 */
 	UFUNCTION(BlueprintCallable, Category = TabList)
-	bool RegisterTab(FName TabNameID, TSubclassOf<UCommonButtonBase> ButtonWidgetType, UWidget* ContentWidget);
+	bool RegisterTab(FName TabNameID, TSubclassOf<UCommonButtonBase> ButtonWidgetType, UWidget* ContentWidget, const int32 TabIndex = -1 /*INDEX_NONE*/);
 
 	UFUNCTION(BlueprintCallable, Category = TabList)
 	bool RemoveTab(FName TabNameID);
@@ -143,6 +151,8 @@ protected:
 
 	virtual void UpdateBindings();
 
+	bool IsRebuildingList() const;
+
 	UFUNCTION(BlueprintImplementableEvent, Category = TabList, meta = (BlueprintProtected = "true"))
 	void HandlePreLinkedSwitcherChanged_BP();
 
@@ -171,7 +181,16 @@ protected:
 	UPROPERTY(EditAnywhere, BlueprintReadOnly, Category = TabList, meta = (ExposeOnSpawn = "true"))
 	bool bAutoListenForInput;
 
+	/**
+	* Whether to defer until next tick rebuilding tab list when inserting new tab (rather than adding to the end).
+	* Useful if inserting multiple tabs in the same tick as the tab list will only be rebuilt once.
+	*/
+	UPROPERTY(EditAnywhere, Category = TabList)
+	bool bDeferRebuildingTabList;
+
 protected:
+	const TMap<FName, FCommonRegisteredTabInfo>& GetRegisteredTabsByID() const;
+
 	UFUNCTION()
 	void HandleTabButtonSelected(UCommonButtonBase* SelectedTabButton, int32 ButtonIndex);
 
@@ -196,14 +215,19 @@ private:
 	void HandleNextTabAction();
 	void HandlePreviousTabAction();
 
+	bool DeferredRebuildTabList(float DeltaTime);
+	void RebuildTabList();
+
 	/** Info about each of the currently registered tabs organized by a given registration name ID */
 	UPROPERTY(Transient)
 	TMap<FName, FCommonRegisteredTabInfo> RegisteredTabsByID;
 
-
 	/** The registration ID of the currently active tab */
 	FName ActiveTabID;
 
+	bool bIsRebuildingList = false;
+	bool bPendingRebuild = false;
+
 	FUIActionBindingHandle NextTabActionHandle;
 	FUIActionBindingHandle PrevTabActionHandle;
-};
\ No newline at end of file
+};
