Common TextBlocks: add options for enabling/disabling scroll style for marqueeing text
- handles ellipsis when scrolling. it should be off and when no longer scrolling it should revert to its desired setting

diff --git a/Engine/Plugins/Experimental/CommonUI/Source/CommonUI/Private/CommonRichTextBlock.cpp b/Engine/Plugins/Experimental/CommonUI/Source/CommonUI/Private/CommonRichTextBlock.cpp
--- a/Engine/Plugins/Experimental/CommonUI/Source/CommonUI/Private/CommonRichTextBlock.cpp
+++ b/Engine/Plugins/Experimental/CommonUI/Source/CommonUI/Private/CommonRichTextBlock.cpp
@@ -502,6 +502,12 @@ void UCommonRichTextBlock::SetText(const FText& InText)
 	}
 }
 
+void UCommonRichTextBlock::SetScrollingEnabled(bool bInIsScrollingEnabled)
+{
+	bIsScrollingEnabled = bInIsScrollingEnabled;
+	SynchronizeProperties();
+}
+
 #if WITH_EDITOR
 
 void UCommonRichTextBlock::OnCreationFromPalette()
@@ -594,6 +600,28 @@ void UCommonRichTextBlock::SynchronizeProperties()
 			ApplyTextBlockScale();
 		}
 	}
+
+	if (MyTextScroller.IsValid())
+	{
+		if (bIsScrollingEnabled)
+		{
+			MyTextScroller->StartScrolling();
+
+			if (MyRichTextBlock.IsValid())
+			{
+				MyRichTextBlock->SetOverflowPolicy(ETextOverflowPolicy::Clip);
+			}
+		}
+		else
+		{
+			MyTextScroller->SuspendScrolling();
+
+			if (MyRichTextBlock.IsValid())
+			{
+				MyRichTextBlock->SetOverflowPolicy(TOptional<ETextOverflowPolicy>());
+			}
+		}
+	}
 }
 
 void UCommonRichTextBlock::ApplyUpdatedDefaultTextStyle()
diff --git a/Engine/Plugins/Experimental/CommonUI/Source/CommonUI/Private/CommonTextBlock.cpp b/Engine/Plugins/Experimental/CommonUI/Source/CommonUI/Private/CommonTextBlock.cpp
--- a/Engine/Plugins/Experimental/CommonUI/Source/CommonUI/Private/CommonTextBlock.cpp
+++ b/Engine/Plugins/Experimental/CommonUI/Source/CommonUI/Private/CommonTextBlock.cpp
@@ -64,16 +64,32 @@ void STextScroller::ResetScrollState()
 	FontAlpha = 1.f;
 	TimeElapsed = 0.f;
 	ScrollOffset = 0.f;
-	ActiveState = EActiveState::EStart;
+	// If suspended, make sure we stay suspended until receiving an explicit call to resume
+	if (IsScrollingEnabled())
+	{
+		ActiveState = EActiveState::EStart;
+	}
 	SetRenderOpacity(1.0f);
 }
 
+void STextScroller::StartScrolling()
+{
+	ActiveState = EActiveState::EStart;
+	ResetScrollState();
+}
+
+void STextScroller::SuspendScrolling()
+{
+	ActiveState = EActiveState::ESuspend;
+	ResetScrollState();
+}
+
 void STextScroller::UpdateTickability(const FGeometry& AllottedGeometry)
 {
 	const float VisibleWidth = AllottedGeometry.GetLocalSize().X;
 	const float DesiredWidth = VisibleWidth == 0 ? VisibleWidth : ChildSlot.GetWidget()->GetDesiredSize().X;
 
-	if (DesiredWidth > (VisibleWidth + 2))
+	if (DesiredWidth > (VisibleWidth + 2) && IsScrollingEnabled())
 	{
 		if (!ActiveTimerHandle.IsValid())
 		{
@@ -409,6 +425,12 @@ void UCommonTextBlock::ResetScrollState()
 	}
 }
 
+void UCommonTextBlock::SetScrollingEnabled(bool bInIsScrollingEnabled)
+{
+	bIsScrollingEnabled = bInIsScrollingEnabled;
+	SynchronizeProperties();
+}
+
 void UCommonTextBlock::SynchronizeProperties()
 {
 	UpdateFromStyle();
@@ -424,6 +446,24 @@ void UCommonTextBlock::SynchronizeProperties()
 	{
 		ApplyFontSizeMultiplier();
 	}
+
+	if (TextScroller.IsValid())
+	{
+		if (bIsScrollingEnabled)
+		{
+			TextScroller->StartScrolling();
+
+			// Hide ellipsis when scrolling
+			if (MyTextBlock.IsValid())
+			{
+				MyTextBlock->SetOverflowPolicy(ETextOverflowPolicy::Clip);
+			}
+		}
+		else
+		{
+			TextScroller->SuspendScrolling();
+		}
+	}
 }
 
 void UCommonTextBlock::SetText(FText InText)
diff --git a/Engine/Plugins/Experimental/CommonUI/Source/CommonUI/Public/CommonRichTextBlock.h b/Engine/Plugins/Experimental/CommonUI/Source/CommonUI/Public/CommonRichTextBlock.h
--- a/Engine/Plugins/Experimental/CommonUI/Source/CommonUI/Public/CommonRichTextBlock.h
+++ b/Engine/Plugins/Experimental/CommonUI/Source/CommonUI/Public/CommonRichTextBlock.h
@@ -41,6 +41,9 @@ public:
 	virtual void ReleaseSlateResources(bool bReleaseChildren) override;
 	virtual void SetText(const FText& InText) override;
 
+	UFUNCTION(BlueprintCallable, Category = "Common Rich Text|Scroll Style")
+	void SetScrollingEnabled(bool bInIsScrollingEnabled);
+
 #if WITH_EDITOR
 	virtual void OnCreationFromPalette() override;
 	const FText GetPaletteCategory() override;
@@ -75,9 +78,13 @@ private:
 	float MobileTextBlockScale = 1.0f;
 
 	/** References the scroll style asset to use, no reference disables scrolling*/
-	UPROPERTY(EditAnywhere, Category = Appearance)
+	UPROPERTY(EditAnywhere, BlueprintReadOnly, Category = Appearance, meta = (ExposeOnSpawn = true, AllowPrivateAccess = true))
 	TSubclassOf<UCommonTextScrollStyle> ScrollStyle;
 
+	/** If scrolling is enabled/disabled initially, this can be updated in blueprint */
+	UPROPERTY(EditAnywhere, BlueprintReadOnly, Category = Appearance, meta = (ExposeOnSpawn = true, AllowPrivateAccess = true))
+	bool bIsScrollingEnabled = true;
+
 	/** True to always display text in ALL CAPS */
 	UPROPERTY(meta = (DeprecatedProperty, DeprecationMessage = "bDisplayAllCaps is deprecated. Please use TextTransformPolicy instead."))
 	bool bDisplayAllCaps_DEPRECATED = false;
diff --git a/Engine/Plugins/Experimental/CommonUI/Source/CommonUI/Public/CommonTextBlock.h b/Engine/Plugins/Experimental/CommonUI/Source/CommonUI/Public/CommonTextBlock.h
--- a/Engine/Plugins/Experimental/CommonUI/Source/CommonUI/Public/CommonTextBlock.h
+++ b/Engine/Plugins/Experimental/CommonUI/Source/CommonUI/Public/CommonTextBlock.h
@@ -128,6 +128,10 @@ public:
 	virtual int32 OnPaint(const FPaintArgs& Args, const FGeometry& AllottedGeometry, const FSlateRect& MyCullingRect, FSlateWindowElementList& OutDrawElements, int32 LayerId, const FWidgetStyle& InWidgetStyle, bool bParentEnabled) const override;
 
 	void ResetScrollState();
+	void StartScrolling();
+	void SuspendScrolling();
+
+	bool IsScrollingEnabled() const { return ActiveState != EActiveState::ESuspend; }
 
 private:
 	void UpdateTickability(const FGeometry& AllottedGeometry);
@@ -143,6 +147,7 @@ private:
 		EStop,
 		EStopWait,
 		EFadeOut,
+		ESuspend,
 	} ActiveState;
 
 	float TimeElapsed;
@@ -181,6 +186,9 @@ public:
 	UFUNCTION(BlueprintCallable, Category = "Common Text|Scroll Style")
 	void ResetScrollState();
 
+	UFUNCTION(BlueprintCallable, Category = "Common Text|Scroll Style")
+	void SetScrollingEnabled(bool bInIsScrollingEnabled);
+
 protected:
 	virtual TSharedRef<SWidget> RebuildWidget() override;
 	virtual void SynchronizeProperties() override;
@@ -200,6 +208,10 @@ private:
 	UPROPERTY(EditAnywhere, BlueprintReadOnly, Category = CommonText, meta = (ExposeOnSpawn = true, AllowPrivateAccess = true))
 	TSubclassOf<UCommonTextScrollStyle> ScrollStyle;
 
+	/** If scrolling is enabled/disabled initially, this can be updated in blueprint */
+	UPROPERTY(EditAnywhere, BlueprintReadOnly, Category = CommonText, meta = (ExposeOnSpawn = true, AllowPrivateAccess = true))
+	bool bIsScrollingEnabled = true;
+
 	/** True to always display text in ALL CAPS */
 	UPROPERTY(meta = (DeprecatedProperty, DeprecationMessage = "bDisplayAllCaps is deprecated. Please use TextTransformPolicy instead."))
 	bool bDisplayAllCaps_DEPRECATED = false;
diff --git a/Engine/Source/Runtime/UMG/Private/Components/TextBlock.cpp b/Engine/Source/Runtime/UMG/Private/Components/TextBlock.cpp
--- a/Engine/Source/Runtime/UMG/Private/Components/TextBlock.cpp
+++ b/Engine/Source/Runtime/UMG/Private/Components/TextBlock.cpp
@@ -137,11 +137,7 @@ void UTextBlock::SetTextTransformPolicy(ETextTransformPolicy InTransformPolicy)
 void UTextBlock::SetTextOverflowPolicy(ETextOverflowPolicy InOverflowPolicy)
 {
 	TextOverflowPolicy = InOverflowPolicy;
-
-	if (MyTextBlock.IsValid())
-	{
-		MyTextBlock->SetOverflowPolicy(TextOverflowPolicy);
-	}
+	SynchronizeProperties();
 }
 
 UMaterialInstanceDynamic* UTextBlock::GetDynamicFontMaterial()
