Common Input:
- Separate the widgets in the action domain system from the old input system
- Input will pass through the active root node (old system) and if not handled, will pass to the action domain system (new system)

diff --git a/Engine/Plugins/Experimental/CommonUI/Source/CommonUI/Private/Input/CommonUIActionRouterBase.cpp b/Engine/Plugins/Experimental/CommonUI/Source/CommonUI/Private/Input/CommonUIActionRouterBase.cpp
--- a/Engine/Plugins/Experimental/CommonUI/Source/CommonUI/Private/Input/CommonUIActionRouterBase.cpp
+++ b/Engine/Plugins/Experimental/CommonUI/Source/CommonUI/Private/Input/CommonUIActionRouterBase.cpp
@@ -355,13 +355,14 @@ ERouteUIInputResult UCommonUIActionRouterBase::ProcessInput(FKey Key, EInputEven
 	
 	if (ProcessHoldResult == EProcessHoldActionResult::Unhandled)
 	{
-		if (bEnableActionDomainRouting)
+		if (bIsActivatableTreeEnabled && ActiveRootNode)
 		{
-			ProcessHoldResult = ProcessHoldInputOnActionDomains(ActiveMode, Key, InputEvent);
+			ProcessHoldResult = ActiveRootNode->ProcessHoldInput(ActiveMode, Key, InputEvent);
 		}
-		else if (bIsActivatableTreeEnabled && ActiveRootNode)
+
+		if (bEnableActionDomainRouting && ProcessHoldResult == EProcessHoldActionResult::Unhandled)
 		{
-			ProcessHoldResult = ActiveRootNode->ProcessHoldInput(ActiveMode, Key, InputEvent);
+			ProcessHoldResult = ProcessHoldInputOnActionDomains(ActiveMode, Key, InputEvent);
 		}
 	}
 
@@ -371,13 +372,14 @@ ERouteUIInputResult UCommonUIActionRouterBase::ProcessInput(FKey Key, EInputEven
 
 			if (!bHandled)
 			{
-				if (bEnableActionDomainRouting)
+				if (bIsActivatableTreeEnabled && ActiveRootNode)
 				{
-					bHandled = ProcessInputOnActionDomains(ActiveMode, Key, Event);
+					bHandled = ActiveRootNode->ProcessNormalInput(ActiveMode, Key, Event);
 				}
-				else if (ActiveRootNode && bIsActivatableTreeEnabled)
+
+				if (bEnableActionDomainRouting && !bHandled)
 				{
-					bHandled = ActiveRootNode->ProcessNormalInput(ActiveMode, Key, Event);
+					bHandled = ProcessInputOnActionDomains(ActiveMode, Key, Event);
 				}
 			}
 
@@ -549,33 +551,33 @@ EMouseCaptureMode UCommonUIActionRouterBase::GetActiveMouseCaptureMode(EMouseCap
 	return ActiveInputConfig.IsSet() ? ActiveInputConfig.GetValue().GetMouseCaptureMode() : DefaultMouseCapture;
 }
 
-void UCommonUIActionRouterBase::HandleRootWidgetSlateReleased(UCommonActivatableWidget* ActivatableWidget)
+void UCommonUIActionRouterBase::HandleRootWidgetSlateReleased(TWeakPtr<FActivatableTreeRoot> WeakRoot)
 {
-	ActivatableWidget->OnSlateReleased().RemoveAll(this);
+	if (!WeakRoot.IsValid())
+	{
+		return;
+	}
 
-	const int32 RootToRemoveIdx = RootNodes.IndexOfByPredicate([ActivatableWidget](const FActivatableTreeRootRef& RootNode) { return RootNode->GetWidget() == ActivatableWidget; });
-	if (ensure(RootToRemoveIdx != INDEX_NONE))
+	FActivatableTreeRootRef RootNode = WeakRoot.Pin().ToSharedRef();
+	if (UCommonActivatableWidget* ActivatableWidget = RootNode->GetWidget())
+	{
+		ActivatableWidget->OnSlateReleased().RemoveAll(this);
+	}
+
+	if (RootNodes.Contains(RootNode))
 	{
 		// It's possible that the widget is destructed as a result of some other deactivation handler, causing us to get here before hearing about
 		// the deactivation. Not a big deal, just need to process the deactivation right here if the node in question is the active root.
-		if (ActiveRootNode == RootNodes[RootToRemoveIdx])
+		if (ActiveRootNode == RootNode)
 		{
-			// @TODO: DarenC - Temp fix for crash on PIE exit
-			if (!ensure(!ActiveRootNode->IsWidgetActivated()) && ActiveRootNode->GetWidget())
+			if (ActiveRootNode->IsWidgetActivated() && ActiveRootNode->GetWidget())
 			{
 				ActiveRootNode->GetWidget()->DeactivateWidget();
 			}
 			HandleRootNodeDeactivated(ActiveRootNode);
 		}
 
-		TWeakPtr<FActivatableTreeRoot> ToBeRemoved = RootNodes[RootToRemoveIdx];
-
-		if (bEnableActionDomainRouting)
-		{
-			RemoveFromActionDomain(RootNodes[RootToRemoveIdx]);
-		}
-
-		RootNodes.RemoveAtSwap(RootToRemoveIdx);
+		RootNodes.RemoveSwap(RootNode);
 		// Cannot actually have this ensure here, because we may be in a function called on the ToBeRemoved node itself, keeping one remaining strong reference
 		// ensureAlways(!ToBeRemoved.IsValid());
 
@@ -586,12 +588,25 @@ void UCommonUIActionRouterBase::HandleRootWidgetSlateReleased(UCommonActivatable
 			ActiveInputConfig.Reset();
 		}
 	}
+	else if (bEnableActionDomainRouting)
+	{
+		for (TPair<TObjectPtr<UCommonInputActionDomain>, FActionDomainSortedRootList>& Pair : ActionDomainRootNodes)
+		{
+			FActionDomainSortedRootList& ActionDomainRootList = Pair.Value;
+			ActionDomainRootList.Remove(RootNode);
+		}
+	}
+	else
+	{
+		UCommonActivatableWidget* ActivatableWidget = RootNode->GetWidget();
+		ensureAlwaysMsgf(false, TEXT("Root node could not be found during deactivation [%s]"), ActivatableWidget ? *ActivatableWidget->GetName() : TEXT("Unknown"));
+	}
 }
 
 void UCommonUIActionRouterBase::HandleRootNodeActivated(TWeakPtr<FActivatableTreeRoot> WeakActivatedRoot)
 {
 	FActivatableTreeRootRef ActivatedRoot = WeakActivatedRoot.Pin().ToSharedRef();
-	if (!ensureAlways(RootNodes.Contains(ActivatedRoot)))
+	if (!RootNodes.Contains(ActivatedRoot))
 	{
 		return;
 	}
@@ -638,36 +653,23 @@ void UCommonUIActionRouterBase::HandlePostGarbageCollect()
 	FUIActionBinding::CleanRegistrations();
 
 	// GC may result in root widget being purged while conditional slate resource release skips HandleRootWidgetSlateReleased, handle this scenario
-	for (auto Iter = RootNodes.CreateIterator(); Iter; ++Iter)
+	for (TArray<FActivatableTreeRootRef>::TIterator Iter = RootNodes.CreateIterator(); Iter; ++Iter)
 	{
 		if (!Iter->Get().IsWidgetValid())
 		{
-			if (bEnableActionDomainRouting)
-			{
-				RemoveFromActionDomain(*Iter);
-			}
 			Iter.RemoveCurrent();
 		}
 	}
-}
-
-void UCommonUIActionRouterBase::AddToActionDomain(FActivatableTreeRootRef RootNode)
-{
-	UCommonInputActionDomain* ActionDomain = RootNode->GetWidget()->GetCalculatedActionDomain();
-
-	if (ActionDomain)
-	{
-		ActionDomainRootNodes.FindOrAdd(ActionDomain).Add(RootNode);
-	}
-}
 
-void UCommonUIActionRouterBase::RemoveFromActionDomain(FActivatableTreeRootRef RootNode)
-{
-	UCommonInputActionDomain* ActionDomain = RootNode->GetWidget()->GetCalculatedActionDomain();
-	FActionDomainSortedRootList* ActionDomainRootList = ActionDomainRootNodes.Find(ActionDomain);
-	if (ActionDomainRootList)
+	for (TPair<TObjectPtr<UCommonInputActionDomain>, FActionDomainSortedRootList>& Pair : ActionDomainRootNodes)
 	{
-		verify(ActionDomainRootList->Remove(RootNode) != INDEX_NONE);
+		for (TArray<FActivatableTreeRootRef>::TIterator Iter = Pair.Value.RootList.CreateIterator(); Iter; ++Iter)
+		{
+			if (!Iter->Get().IsWidgetValid())
+			{
+				Iter.RemoveCurrent();
+			}
+		}
 	}
 }
 
@@ -769,7 +771,13 @@ void UCommonUIActionRouterBase::ProcessRebuiltWidgets()
 	{
 		if (RebuiltWidget.IsValid() && RebuiltWidget->GetCachedWidget())
 		{
-			if (UCommonActivatableWidget* ActivatableParent = !RebuiltWidget->IsModal() ? ::FindOwningActivatable(*RebuiltWidget) : nullptr)
+			if (UCommonInputActionDomain* ActionDomain = bEnableActionDomainRouting ? RebuiltWidget->GetCalculatedActionDomain() : nullptr)
+			{
+				FActivatableTreeRootRef RootNode = FActivatableTreeRoot::Create(*this, *RebuiltWidget);
+				RebuiltWidget->OnSlateReleased().AddUObject(this, &UCommonUIActionRouterBase::HandleRootWidgetSlateReleased, TWeakPtr<FActivatableTreeRoot>(RootNode));
+				ActionDomainRootNodes.FindOrAdd(ActionDomain).Add(RootNode);
+			}
+			else if (UCommonActivatableWidget* ActivatableParent = !RebuiltWidget->IsModal() ? ::FindOwningActivatable(*RebuiltWidget) : nullptr)
 			{
 				WidgetsByDirectParent.FindOrAdd(ActivatableParent).Add(RebuiltWidget.Get());
 			}
@@ -789,13 +797,9 @@ void UCommonUIActionRouterBase::ProcessRebuiltWidgets()
 		TWeakPtr<FActivatableTreeRoot> WeakRoot(RootNode);
 		RootNode->OnActivated.BindUObject(this, &UCommonUIActionRouterBase::HandleRootNodeActivated, WeakRoot);
 		RootNode->OnDeactivated.BindUObject(this, &UCommonUIActionRouterBase::HandleRootNodeDeactivated, WeakRoot);
-		RootWidget->OnSlateReleased().AddUObject(this, &UCommonUIActionRouterBase::HandleRootWidgetSlateReleased, RootWidget);
-		RootNodes.Add(RootNode);
+		RootWidget->OnSlateReleased().AddUObject(this, &UCommonUIActionRouterBase::HandleRootWidgetSlateReleased, WeakRoot);
 
-		if (bEnableActionDomainRouting)
-		{
-			AddToActionDomain(RootNode);
-		}
+		RootNodes.Add(RootNode);
 
 		AssembleTreeRecursive(RootNode, WidgetsByDirectParent);
 
@@ -914,7 +918,10 @@ bool UCommonUIActionRouterBase::Tick(float DeltaTime)
 		const EProcessHoldActionResult ProcessHoldResult = PersistentActions->ProcessHoldInput(ActiveMode, HeldKey, EInputEvent::IE_Repeat);
 		if (bIsActivatableTreeEnabled && ActiveRootNode && ProcessHoldResult == EProcessHoldActionResult::Unhandled)
 		{
-			ActiveRootNode->ProcessHoldInput(ActiveMode, HeldKey, EInputEvent::IE_Repeat);
+			if (ActiveRootNode->ProcessHoldInput(ActiveMode, HeldKey, EInputEvent::IE_Repeat) == EProcessHoldActionResult::Unhandled)
+			{
+				ProcessHoldInputOnActionDomains(ActiveMode, HeldKey, EInputEvent::IE_Repeat);
+			}
 		}
 	}
 
@@ -1103,6 +1110,20 @@ FActivatableTreeNodePtr UCommonUIActionRouterBase::FindNode(const UCommonActivat
 				break;
 			}
 		}
+
+		if (bEnableActionDomainRouting && !FoundNode.IsValid())
+		{
+			for (const TPair<TObjectPtr<UCommonInputActionDomain>, FActionDomainSortedRootList>& Pair : ActionDomainRootNodes)
+			{
+				for (const FActivatableTreeRootRef RootNode : Pair.Value.RootList)
+				{
+					if (Widget == RootNode->GetWidget())
+					{
+						return RootNode;
+					}
+				}
+			}
+		}
 	}
 
 	return FoundNode;
diff --git a/Engine/Plugins/Experimental/CommonUI/Source/CommonUI/Public/CommonActivatableWidget.h b/Engine/Plugins/Experimental/CommonUI/Source/CommonUI/Public/CommonActivatableWidget.h
--- a/Engine/Plugins/Experimental/CommonUI/Source/CommonUI/Public/CommonActivatableWidget.h
+++ b/Engine/Plugins/Experimental/CommonUI/Source/CommonUI/Public/CommonActivatableWidget.h
@@ -185,13 +185,13 @@ protected:
 	UPROPERTY(EditAnywhere, Category = Activation, meta = (EditCondition = bSupportsActivationFocus))
 	bool bAutoRestoreFocus = false;
 
-	UPROPERTY(EditAnywhere, BlueprintReadOnly, Category = "Input | ActionDomain", meta = (InlineEditConditionToggle))
+	UPROPERTY(EditAnywhere, BlueprintReadOnly, Category = "Input", meta = (InlineEditConditionToggle))
 	bool bOverrideActionDomain = false;
 
 	/**
 	 * Enable to override the inherited ActionDomain from owning CommonActivatableWidget.
 	 */
-	UPROPERTY(EditAnywhere, BlueprintReadOnly, Category = "Input | ActionDomain", meta = (EditCondition = "bOverrideActionDomain"))
+	UPROPERTY(EditAnywhere, BlueprintReadOnly, Category = "Input", meta = (EditCondition = "bOverrideActionDomain"))
 	TSoftObjectPtr<UCommonInputActionDomain> ActionDomainOverride;
 
 private:
diff --git a/Engine/Plugins/Experimental/CommonUI/Source/CommonUI/Public/Input/CommonUIActionRouterBase.h b/Engine/Plugins/Experimental/CommonUI/Source/CommonUI/Public/Input/CommonUIActionRouterBase.h
--- a/Engine/Plugins/Experimental/CommonUI/Source/CommonUI/Public/Input/CommonUIActionRouterBase.h
+++ b/Engine/Plugins/Experimental/CommonUI/Source/CommonUI/Public/Input/CommonUIActionRouterBase.h
@@ -150,7 +150,7 @@ private:
 	void ProcessRebuiltWidgets();
 	void AssembleTreeRecursive(const FActivatableTreeNodeRef& CurNode, TMap<UCommonActivatableWidget*, TArray<UCommonActivatableWidget*>>& WidgetsByDirectParent);
 
-	void HandleRootWidgetSlateReleased(UCommonActivatableWidget* ActivatableWidget);
+	void HandleRootWidgetSlateReleased(TWeakPtr<FActivatableTreeRoot> WeakRoot);
 	void HandleRootNodeActivated(TWeakPtr<FActivatableTreeRoot> WeakActivatedRoot);
 	void HandleRootNodeDeactivated(TWeakPtr<FActivatableTreeRoot> WeakDeactivatedRoot);
 	void HandleLeafmostActiveNodeChanged();
@@ -159,8 +159,6 @@ private:
 
 	void HandlePostGarbageCollect();
 
-	void AddToActionDomain(FActivatableTreeRootRef RootNode);
-	void RemoveFromActionDomain(FActivatableTreeRootRef RootNode);
 	bool ProcessInputOnActionDomains(ECommonInputMode ActiveInputMode, FKey Key, EInputEvent InputEvent) const;
 	EProcessHoldActionResult ProcessHoldInputOnActionDomains(ECommonInputMode ActiveInputMode, FKey Key, EInputEvent InputEvent) const;
 
@@ -209,7 +207,7 @@ private:
 	mutable TArray<FKey> HeldKeys;
 
 	/** A wrapper around TArray that keeps RootList sorted by PaintLayer during insertion */
-	struct FActionDomainSortedRootList 
+	struct FActionDomainSortedRootList
 	{
 		TArray<FActivatableTreeRootRef> RootList;
 
