CIS Fix

diff --git a/Engine/Plugins/Experimental/CommonUI/Source/CommonUI/Private/Input/CommonUIActionRouterBase.cpp b/Engine/Plugins/Experimental/CommonUI/Source/CommonUI/Private/Input/CommonUIActionRouterBase.cpp
--- a/Engine/Plugins/Experimental/CommonUI/Source/CommonUI/Private/Input/CommonUIActionRouterBase.cpp
+++ b/Engine/Plugins/Experimental/CommonUI/Source/CommonUI/Private/Input/CommonUIActionRouterBase.cpp
@@ -1115,7 +1115,7 @@ FActivatableTreeNodePtr UCommonUIActionRouterBase::FindNode(const UCommonActivat
 		{
 			for (const TPair<TObjectPtr<UCommonInputActionDomain>, FActionDomainSortedRootList>& Pair : ActionDomainRootNodes)
 			{
-				for (const FActivatableTreeRootRef RootNode : Pair.Value.RootList)
+				for (const FActivatableTreeRootRef& RootNode : Pair.Value.RootList)
 				{
 					if (Widget == RootNode->GetWidget())
 					{
