Crash fix - In ActionRouterBase Tick function, the call to ActionDomain input processor was not guarded by a check to see if ActionDomainRouting is enabled. Also add sanity checks in the ProcessInput functions so it's not a fatal crash if it happens again.

diff --git a/Engine/Plugins/Experimental/CommonUI/Source/CommonUI/Private/Input/CommonUIActionRouterBase.cpp b/Engine/Plugins/Experimental/CommonUI/Source/CommonUI/Private/Input/CommonUIActionRouterBase.cpp
--- a/Engine/Plugins/Experimental/CommonUI/Source/CommonUI/Private/Input/CommonUIActionRouterBase.cpp
+++ b/Engine/Plugins/Experimental/CommonUI/Source/CommonUI/Private/Input/CommonUIActionRouterBase.cpp
@@ -678,6 +678,11 @@ bool UCommonUIActionRouterBase::ProcessInputOnActionDomains(ECommonInputMode Act
 	UCommonInputSubsystem& CommonInputSubsystem = GetInputSubsystem();
 	UCommonInputActionDomainTable* ActionDomainTable = CommonInputSubsystem.GetActionDomainTable();
 
+	if (!ActionDomainTable)
+	{
+		return false;
+	}
+
 	bool bInputEventHandledAtLeastOnce = false;
 
 	for (const UCommonInputActionDomain* ActionDomain : ActionDomainTable->ActionDomains)
@@ -723,6 +728,11 @@ EProcessHoldActionResult UCommonUIActionRouterBase::ProcessHoldInputOnActionDoma
 	UCommonInputSubsystem& CommonInputSubsystem = GetInputSubsystem();
 	UCommonInputActionDomainTable* ActionDomainTable = CommonInputSubsystem.GetActionDomainTable();
 
+	if (!ActionDomainTable)
+	{
+		return EProcessHoldActionResult::Unhandled;
+	}
+
 	EProcessHoldActionResult OutHoldActionResult = EProcessHoldActionResult::Unhandled;
 
 	for (const UCommonInputActionDomain* ActionDomain : ActionDomainTable->ActionDomains)
@@ -920,7 +930,10 @@ bool UCommonUIActionRouterBase::Tick(float DeltaTime)
 		{
 			if (ActiveRootNode->ProcessHoldInput(ActiveMode, HeldKey, EInputEvent::IE_Repeat) == EProcessHoldActionResult::Unhandled)
 			{
-				ProcessHoldInputOnActionDomains(ActiveMode, HeldKey, EInputEvent::IE_Repeat);
+				if (bEnableActionDomainRouting)
+				{
+					ProcessHoldInputOnActionDomains(ActiveMode, HeldKey, EInputEvent::IE_Repeat);
+				}
 			}
 		}
 	}
