Added ability to check if input is allowed.

diff --git a/Engine/Plugins/Experimental/CommonUI/Source/CommonUI/Private/CommonUISubsystemBase.cpp b/Engine/Plugins/Experimental/CommonUI/Source/CommonUI/Private/CommonUISubsystemBase.cpp
--- a/Engine/Plugins/Experimental/CommonUI/Source/CommonUI/Private/CommonUISubsystemBase.cpp
+++ b/Engine/Plugins/Experimental/CommonUI/Source/CommonUI/Private/CommonUISubsystemBase.cpp
@@ -127,4 +127,17 @@ void UCommonUISubsystemBase::SetInputAllowed(bool bEnabled, const FName& Reason,
 		InputSubsystem->SetInputTypeFilter(ECommonInputType::Gamepad, Reason, !bEnabled);
 		InputSubsystem->SetInputTypeFilter(ECommonInputType::Touch, Reason, !bEnabled);
 	}
-}
\ No newline at end of file
+}
+
+bool UCommonUISubsystemBase::IsInputAllowed(const ULocalPlayer* LocalPlayer) const
+{
+	if (UCommonInputSubsystem* InputSubSystem = UCommonInputSubsystem::Get(LocalPlayer))
+	{
+		const bool bKeyboardAndMouseFiltered = InputSubSystem->GetInputTypeFilter(ECommonInputType::MouseAndKeyboard);
+		const bool bGamepadFiltered = InputSubSystem->GetInputTypeFilter(ECommonInputType::Gamepad);
+		const bool bTouchFiltered = InputSubSystem->GetInputTypeFilter(ECommonInputType::Touch);
+		return !bKeyboardAndMouseFiltered || !bGamepadFiltered || !bTouchFiltered;
+	}
+
+	return true;
+}
diff --git a/Engine/Plugins/Experimental/CommonUI/Source/CommonUI/Public/CommonUISubsystemBase.h b/Engine/Plugins/Experimental/CommonUI/Source/CommonUI/Public/CommonUISubsystemBase.h
--- a/Engine/Plugins/Experimental/CommonUI/Source/CommonUI/Public/CommonUISubsystemBase.h
+++ b/Engine/Plugins/Experimental/CommonUI/Source/CommonUI/Public/CommonUISubsystemBase.h
@@ -26,8 +26,6 @@ public:
 	/** set the analytic provider for the CommonUI Widgets to use */
 	void SetAnalyticProvider(const TSharedPtr<IAnalyticsProviderET>& AnalyticProvider);
 
-public:
-
 	// Gets Action Button Icon for current gamepad
 	UFUNCTION(BlueprintCallable, Category = CommonUISubsystem)
 	FSlateBrush GetInputActionButtonIcon(const FDataTableRowHandle& InputActionRowHandle, ECommonInputType InputType, const FName& GamepadName) const;
@@ -40,8 +38,8 @@ public:
 	//CommonUI.PanelPushed
 	void FireEvent_PanelPushed(const FString& PanelName) const;
 	
-public:
 	virtual void SetInputAllowed(bool bEnabled, const FName& Reason, const ULocalPlayer& LocalPlayer);
+	virtual bool IsInputAllowed(const ULocalPlayer* LocalPlayer) const;
 
 private:
 
