CommonTileView: move SCommonTileView to a header so we can inherit from it properly
- minor cleanup: removed override that was just calling parent, add parent call for OnFocusLost

diff --git a/Engine/Plugins/Experimental/CommonUI/Source/CommonUI/Private/CommonTileView.cpp b/Engine/Plugins/Experimental/CommonUI/Source/CommonUI/Private/CommonTileView.cpp
--- a/Engine/Plugins/Experimental/CommonUI/Source/CommonUI/Private/CommonTileView.cpp
+++ b/Engine/Plugins/Experimental/CommonUI/Source/CommonUI/Private/CommonTileView.cpp
@@ -3,85 +3,7 @@
 #include "CommonTileView.h"
 #include "CommonUIPrivatePCH.h"
 #include "SCommonButtonTableRow.h"
-
-///////////////////////
-// SCommonTileView
-///////////////////////
-
-template <typename ItemType>
-class SCommonTileView : public STileView<ItemType>
-{
-public:
-	virtual FReply OnFocusReceived(const FGeometry& MyGeometry, const FFocusEvent& InFocusEvent) override
-	{
-		if (bScrollToSelectedOnFocus && (InFocusEvent.GetCause() == EFocusCause::Navigation || InFocusEvent.GetCause() == EFocusCause::SetDirectly))
-		{
-			if (this->ItemsSource && this->ItemsSource->Num() > 0)
-			{
-				if(this->GetNumItemsSelected() == 0)
-				{
-					ItemType FirstItem = (*this->ItemsSource)[0];
-					this->SetSelection(FirstItem, ESelectInfo::OnNavigation);
-					this->RequestNavigateToItem(FirstItem, InFocusEvent.GetUser());
-				}
-				else
-				{
-					TArray<ItemType> ItemArray;
-					this->GetSelectedItems(ItemArray);
-
-					ItemType FirstSelected = ItemArray[0];
-					this->SetSelection(FirstSelected, ESelectInfo::OnNavigation);
-					this->RequestNavigateToItem(FirstSelected, InFocusEvent.GetUser());
-				}
-			}
-		}
-		bScrollToSelectedOnFocus = true;
-		return STileView<ItemType>::OnFocusReceived(MyGeometry, InFocusEvent);
-	}
-
-	virtual void OnFocusLost(const FFocusEvent& InFocusEvent) override
-	{
-		// if we've lost focus, even if we're already scrolling a tile into view, refrain from navigating to that tile and thus stealing focus from another widget
-		this->bNavigateOnScrollIntoView = false;
-	}
-
-	virtual void OnMouseLeave(const FPointerEvent& MouseEvent) override
-	{
-		STileView<ItemType>::OnMouseLeave(MouseEvent);
-
-		if (MouseEvent.IsTouchEvent() && this->HasMouseCapture())
-		{
-			// Regular list views will clear this flag when the pointer leaves the list. To
-			// continue scrolling outside the list, we need this to remain on.
-			this->bStartedTouchInteraction = true;
-		}
-	}
-
-	virtual FReply OnTouchMoved(const FGeometry& MyGeometry, const FPointerEvent& InTouchEvent) override
-	{
-		FReply Reply = STileView<ItemType>::OnTouchMoved(MyGeometry, InTouchEvent);
-		
-		if (Reply.IsEventHandled() && this->HasMouseCapture())
-		{
-			bScrollToSelectedOnFocus = false;
-			Reply.SetUserFocus(this->AsShared());
-		}
-		return Reply;
-	}
-
-	virtual FReply OnTouchEnded(const FGeometry& MyGeometry, const FPointerEvent& InTouchEvent) override
-	{
-		return STileView<ItemType>::OnTouchEnded(MyGeometry, InTouchEvent);
-	}
-
-private:
-	bool bScrollToSelectedOnFocus = true;
-};
-
-
-///////////////////////
-// UCommonTileView
-///////////////////////
+#include "SCommonTileView.h"
 
 UCommonTileView::UCommonTileView(const FObjectInitializer& ObjectInitializer)
 	: Super(ObjectInitializer)
diff --git a/Engine/Plugins/Experimental/CommonUI/Source/CommonUI/Public/SCommonTileView.h b/Engine/Plugins/Experimental/CommonUI/Source/CommonUI/Public/SCommonTileView.h
new file mode 100644
--- /dev/null
+++ b/Engine/Plugins/Experimental/CommonUI/Source/CommonUI/Public/SCommonTileView.h
@@ -0,0 +1,72 @@
+// Copyright Epic Games, Inc. All Rights Reserved.
+
+#pragma once
+
+#include "Widgets/Views/STileView.h"
+
+template <typename ItemType>
+class COMMONUI_API SCommonTileView : public STileView<ItemType>
+{
+public:
+	virtual FReply OnFocusReceived(const FGeometry& MyGeometry, const FFocusEvent& InFocusEvent) override
+	{
+		if (bScrollToSelectedOnFocus && (InFocusEvent.GetCause() == EFocusCause::Navigation || InFocusEvent.GetCause() == EFocusCause::SetDirectly))
+		{
+			if (this->ItemsSource && this->ItemsSource->Num() > 0)
+			{
+				if (this->GetNumItemsSelected() == 0)
+				{
+					ItemType FirstItem = (*this->ItemsSource)[0];
+					this->SetSelection(FirstItem, ESelectInfo::OnNavigation);
+					this->RequestNavigateToItem(FirstItem, InFocusEvent.GetUser());
+				}
+				else
+				{
+					TArray<ItemType> ItemArray;
+					this->GetSelectedItems(ItemArray);
+
+					ItemType FirstSelected = ItemArray[0];
+					this->SetSelection(FirstSelected, ESelectInfo::OnNavigation);
+					this->RequestNavigateToItem(FirstSelected, InFocusEvent.GetUser());
+				}
+			}
+		}
+		bScrollToSelectedOnFocus = true;
+		return STileView<ItemType>::OnFocusReceived(MyGeometry, InFocusEvent);
+	}
+
+	virtual void OnFocusLost(const FFocusEvent& InFocusEvent) override
+	{
+		STileView<ItemType>::OnFocusLost(InFocusEvent);
+
+		// if we've lost focus, even if we're already scrolling a tile into view, refrain from navigating to that tile and thus stealing focus from another widget
+		this->bNavigateOnScrollIntoView = false;
+	}
+
+	virtual void OnMouseLeave(const FPointerEvent& MouseEvent) override
+	{
+		STileView<ItemType>::OnMouseLeave(MouseEvent);
+
+		if (MouseEvent.IsTouchEvent() && this->HasMouseCapture())
+		{
+			// Regular list views will clear this flag when the pointer leaves the list. To
+			// continue scrolling outside the list, we need this to remain on.
+			this->bStartedTouchInteraction = true;
+		}
+	}
+
+	virtual FReply OnTouchMoved(const FGeometry& MyGeometry, const FPointerEvent& InTouchEvent) override
+	{
+		FReply Reply = STileView<ItemType>::OnTouchMoved(MyGeometry, InTouchEvent);
+
+		if (Reply.IsEventHandled() && this->HasMouseCapture())
+		{
+			bScrollToSelectedOnFocus = false;
+			Reply.SetUserFocus(this->AsShared());
+		}
+		return Reply;
+	}
+
+private:
+	bool bScrollToSelectedOnFocus = true;
+};
