Allow CommonAnimatedSwitcher to report transitioning animation status.
#OKFORGITHUB PUBLIC

diff --git a/Engine/Plugins/Experimental/CommonUI/Source/CommonUI/Private/CommonAnimatedSwitcher.cpp b/Engine/Plugins/Experimental/CommonUI/Source/CommonUI/Private/CommonAnimatedSwitcher.cpp
--- a/Engine/Plugins/Experimental/CommonUI/Source/CommonUI/Private/CommonAnimatedSwitcher.cpp
+++ b/Engine/Plugins/Experimental/CommonUI/Source/CommonUI/Private/CommonAnimatedSwitcher.cpp
@@ -141,6 +141,8 @@ void UCommonAnimatedSwitcher::HandleSlateIsTransitioningChanged(bool bIsTransiti
 {
 	// While the switcher is transitioning, put up the guard to intercept all input
 	MyInputGuard->SetVisibility(bIsTransitioning ? EVisibility::Visible : EVisibility::Collapsed);
+
+	OnTransitioningChanged.Broadcast(bIsTransitioning);
 }
 
 void UCommonAnimatedSwitcher::SetActiveWidgetIndex_Internal(int32 Index)
@@ -177,3 +179,15 @@ void UCommonAnimatedSwitcher::SetActiveWidgetIndex_Internal(int32 Index)
 		bSetOnce = true;
 	}
 }
+
+bool UCommonAnimatedSwitcher::IsTransitionPlaying() const
+{
+	if (MyAnimatedSwitcher.IsValid())
+	{
+		return MyAnimatedSwitcher->IsTransitionPlaying();
+	}
+	else
+	{
+		return false;
+	}
+}
diff --git a/Engine/Plugins/Experimental/CommonUI/Source/CommonUI/Private/Slate/SCommonAnimatedSwitcher.cpp b/Engine/Plugins/Experimental/CommonUI/Source/CommonUI/Private/Slate/SCommonAnimatedSwitcher.cpp
--- a/Engine/Plugins/Experimental/CommonUI/Source/CommonUI/Private/Slate/SCommonAnimatedSwitcher.cpp
+++ b/Engine/Plugins/Experimental/CommonUI/Source/CommonUI/Private/Slate/SCommonAnimatedSwitcher.cpp
@@ -231,4 +231,9 @@ float SCommonAnimatedSwitcher::GetTransitionProgress() const
 	}
 		
 	return Progress;
+}
+
+bool SCommonAnimatedSwitcher::IsTransitionPlaying() const
+{
+	return TransitionSequence.IsPlaying();
 }
\ No newline at end of file
diff --git a/Engine/Plugins/Experimental/CommonUI/Source/CommonUI/Public/CommonAnimatedSwitcher.h b/Engine/Plugins/Experimental/CommonUI/Source/CommonUI/Public/CommonAnimatedSwitcher.h
--- a/Engine/Plugins/Experimental/CommonUI/Source/CommonUI/Public/CommonAnimatedSwitcher.h
+++ b/Engine/Plugins/Experimental/CommonUI/Source/CommonUI/Public/CommonAnimatedSwitcher.h
@@ -46,6 +46,10 @@ public:
 	UFUNCTION(BlueprintCallable, Category = "Common Widget Switcher")
 	bool IsCurrentlySwitching() const;
 
+	/** Is the switcher playing a transition animation? */
+	UFUNCTION(BlueprintCallable, Category = "Common Widget Switcher")
+	bool IsTransitionPlaying() const;
+
 protected:
 	virtual void HandleSlateActiveIndexChanged(int32 ActiveIndex);
 
@@ -60,6 +64,10 @@ public:
 	DECLARE_EVENT_TwoParams(UCommonAnimatedSwitcher, FOnActiveIndexChanged, UWidget*, int32)
 	FOnActiveIndexChanged OnActiveWidgetIndexChanged;
 
+	/** Fires when the switcher changes its transition animation state */
+	DECLARE_EVENT_OneParam(UCommonAnimatedSwitcher, FOnTransitioningChanged, bool)
+	FOnTransitioningChanged OnTransitioningChanged;
+
 protected:
 	/** The type of transition to play between widgets */
 	UPROPERTY(EditAnywhere, BlueprintReadOnly, Category = "Transition")
diff --git a/Engine/Plugins/Experimental/CommonUI/Source/CommonUI/Public/Slate/SCommonAnimatedSwitcher.h b/Engine/Plugins/Experimental/CommonUI/Source/CommonUI/Public/Slate/SCommonAnimatedSwitcher.h
--- a/Engine/Plugins/Experimental/CommonUI/Source/CommonUI/Public/Slate/SCommonAnimatedSwitcher.h
+++ b/Engine/Plugins/Experimental/CommonUI/Source/CommonUI/Public/Slate/SCommonAnimatedSwitcher.h
@@ -89,6 +89,8 @@ public:
 
 	void SetTransition(float Duration, ETransitionCurve Curve);
 
+	bool IsTransitionPlaying() const;
+
 private:
 	EActiveTimerReturnType UpdateTransition(double InCurrentTime, float InDeltaTime);
 	float GetTransitionProgress() const;
