UMG: Update the TextBlock widget with accessors.

diff --git a/Engine/Plugins/Experimental/CommonUI/Source/CommonUI/Private/CommonNumericTextBlock.cpp b/Engine/Plugins/Experimental/CommonUI/Source/CommonUI/Private/CommonNumericTextBlock.cpp
--- a/Engine/Plugins/Experimental/CommonUI/Source/CommonUI/Private/CommonNumericTextBlock.cpp
+++ b/Engine/Plugins/Experimental/CommonUI/Source/CommonUI/Private/CommonNumericTextBlock.cpp
@@ -129,11 +129,13 @@ bool UCommonNumericTextBlock::CanEditChange(const FProperty* InProperty) const
 	if (bIsEditable && InProperty)
 	{
 		const FName PropertyName = InProperty->GetFName();
-
+PRAGMA_DISABLE_DEPRECATION_WARNINGS
+		// this should be move in a customizaer
 		if (PropertyName == GET_MEMBER_NAME_CHECKED(UTextBlock, MinDesiredWidth))
 		{
 			bIsEditable = PerformSizeInterpolation != true;
 		}
+PRAGMA_ENABLE_DEPRECATION_WARNINGS
 	}
 
 	return bIsEditable;
diff --git a/Engine/Plugins/Experimental/CommonUI/Source/CommonUI/Private/CommonTextBlock.cpp b/Engine/Plugins/Experimental/CommonUI/Source/CommonUI/Private/CommonTextBlock.cpp
--- a/Engine/Plugins/Experimental/CommonUI/Source/CommonUI/Private/CommonTextBlock.cpp
+++ b/Engine/Plugins/Experimental/CommonUI/Source/CommonUI/Private/CommonTextBlock.cpp
@@ -303,18 +303,77 @@ void UCommonTextBlock::PostLoad()
 			bool bAllDefaults = true;
 
 			// Font
-			if (Font.FontObject == CDO->Font.FontObject) { Font.FontObject = DefaultStyleCDO->Font.FontObject; } else { bAllDefaults = false; }
-			if (Font.FontMaterial == CDO->Font.FontMaterial) { Font.FontMaterial = DefaultStyleCDO->Font.FontMaterial; } else { bAllDefaults = false; }
-			if (Font.OutlineSettings.OutlineSize == CDO->Font.OutlineSettings.OutlineSize) { Font.OutlineSettings.OutlineSize = DefaultStyleCDO->Font.OutlineSettings.OutlineSize; } else { bAllDefaults = false; }
-			if (Font.OutlineSettings.bSeparateFillAlpha == CDO->Font.OutlineSettings.bSeparateFillAlpha) { Font.OutlineSettings.bSeparateFillAlpha = DefaultStyleCDO->Font.OutlineSettings.bSeparateFillAlpha; } else { bAllDefaults = false; }
-			if (Font.OutlineSettings.bApplyOutlineToDropShadows == CDO->Font.OutlineSettings.bApplyOutlineToDropShadows) { Font.OutlineSettings.bApplyOutlineToDropShadows = DefaultStyleCDO->Font.OutlineSettings.bApplyOutlineToDropShadows; } else { bAllDefaults = false; }
-			if (Font.OutlineSettings.OutlineMaterial == CDO->Font.OutlineSettings.OutlineMaterial) { Font.OutlineSettings.OutlineMaterial = DefaultStyleCDO->Font.OutlineSettings.OutlineMaterial; } else { bAllDefaults = false; }
-			if (Font.OutlineSettings.OutlineColor.R == CDO->Font.OutlineSettings.OutlineColor.R) { Font.OutlineSettings.OutlineColor.R = DefaultStyleCDO->Font.OutlineSettings.OutlineColor.R; } else { bAllDefaults = false; }
-			if (Font.OutlineSettings.OutlineColor.G == CDO->Font.OutlineSettings.OutlineColor.G) { Font.OutlineSettings.OutlineColor.G = DefaultStyleCDO->Font.OutlineSettings.OutlineColor.G; } else { bAllDefaults = false; }
-			if (Font.OutlineSettings.OutlineColor.B == CDO->Font.OutlineSettings.OutlineColor.B) { Font.OutlineSettings.OutlineColor.B = DefaultStyleCDO->Font.OutlineSettings.OutlineColor.B; } else { bAllDefaults = false; }
-			if (Font.OutlineSettings.OutlineColor.A == CDO->Font.OutlineSettings.OutlineColor.A) { Font.OutlineSettings.OutlineColor.A = DefaultStyleCDO->Font.OutlineSettings.OutlineColor.A; } else { bAllDefaults = false; }
-			if (Font.TypefaceFontName == CDO->Font.TypefaceFontName) { Font.TypefaceFontName = DefaultStyleCDO->Font.TypefaceFontName; } else { bAllDefaults = false; };
-			if (Font.Size == CDO->Font.Size) { Font.Size = DefaultStyleCDO->Font.Size; } else { bAllDefaults = false; };
+			{
+				bool bFontHasChanged = false;
+				FSlateFontInfo MyFont = GetFont();
+				const FSlateFontInfo& CDOFont = CDO->GetFont();
+				if (MyFont.FontObject == CDOFont.FontObject)
+				{
+					MyFont.FontObject = DefaultStyleCDO->Font.FontObject;
+					bFontHasChanged = true;
+				}
+				if (MyFont.FontMaterial == CDOFont.FontMaterial)
+				{
+					MyFont.FontMaterial = DefaultStyleCDO->Font.FontMaterial;
+					bFontHasChanged = true;
+				}
+				if (MyFont.OutlineSettings.OutlineSize == CDOFont.OutlineSettings.OutlineSize)
+				{
+					MyFont.OutlineSettings.OutlineSize = DefaultStyleCDO->Font.OutlineSettings.OutlineSize;
+					bFontHasChanged = true;
+				}
+				if (MyFont.OutlineSettings.bSeparateFillAlpha == CDOFont.OutlineSettings.bSeparateFillAlpha)
+				{
+					MyFont.OutlineSettings.bSeparateFillAlpha = DefaultStyleCDO->Font.OutlineSettings.bSeparateFillAlpha;
+					bFontHasChanged = true;
+				}
+				if (MyFont.OutlineSettings.bApplyOutlineToDropShadows == CDOFont.OutlineSettings.bApplyOutlineToDropShadows)
+				{
+					MyFont.OutlineSettings.bApplyOutlineToDropShadows = DefaultStyleCDO->Font.OutlineSettings.bApplyOutlineToDropShadows;
+					bFontHasChanged = true;
+				}
+				if (MyFont.OutlineSettings.OutlineMaterial == CDOFont.OutlineSettings.OutlineMaterial)
+				{
+					MyFont.OutlineSettings.OutlineMaterial = DefaultStyleCDO->Font.OutlineSettings.OutlineMaterial;
+					bFontHasChanged = true;
+				}
+				if (MyFont.OutlineSettings.OutlineColor.R == CDOFont.OutlineSettings.OutlineColor.R)
+				{
+					MyFont.OutlineSettings.OutlineColor.R = DefaultStyleCDO->Font.OutlineSettings.OutlineColor.R;
+					bFontHasChanged = true;
+				}
+				if (MyFont.OutlineSettings.OutlineColor.G == CDOFont.OutlineSettings.OutlineColor.G)
+				{
+					MyFont.OutlineSettings.OutlineColor.G = DefaultStyleCDO->Font.OutlineSettings.OutlineColor.G;
+					bFontHasChanged = true;
+				}
+				if (MyFont.OutlineSettings.OutlineColor.B == CDOFont.OutlineSettings.OutlineColor.B)
+				{
+					MyFont.OutlineSettings.OutlineColor.B = DefaultStyleCDO->Font.OutlineSettings.OutlineColor.B;
+					bFontHasChanged = true;
+				}
+				if (MyFont.OutlineSettings.OutlineColor.A == CDOFont.OutlineSettings.OutlineColor.A)
+				{
+					MyFont.OutlineSettings.OutlineColor.A = DefaultStyleCDO->Font.OutlineSettings.OutlineColor.A;
+					bFontHasChanged = true;
+				}
+				if (MyFont.TypefaceFontName == CDOFont.TypefaceFontName)
+				{
+					MyFont.TypefaceFontName = DefaultStyleCDO->Font.TypefaceFontName;
+					bFontHasChanged = true;
+				}
+				if (MyFont.Size == CDOFont.Size)
+				{
+					MyFont.Size = DefaultStyleCDO->Font.Size;
+					bFontHasChanged = true;
+				}
+
+				if (bFontHasChanged)
+				{
+					SetFont(MyFont);
+				}
+				bAllDefaults = bAllDefaults && !bFontHasChanged;
+			}
 
 			// Margin
 			if (Margin.Left == CDO->Margin.Left) { Margin.Left = DefaultStyleCDO->Margin.Left; } else { bAllDefaults = false; }
@@ -326,17 +385,63 @@ void UCommonTextBlock::PostLoad()
 			if (LineHeightPercentage == CDO->LineHeightPercentage) { LineHeightPercentage = DefaultStyleCDO->LineHeightPercentage; } else { bAllDefaults = false; }
 
 			// ColorAndOpacity
-			if (ColorAndOpacity == CDO->ColorAndOpacity) { ColorAndOpacity = DefaultStyleCDO->Color; } else { bAllDefaults = false; }
+			if (GetColorAndOpacity() == CDO->GetColorAndOpacity()) { SetColorAndOpacity(DefaultStyleCDO->Color); } else { bAllDefaults = false; }
 
 			// ShadowOffset
-			if (ShadowOffset.X == CDO->ShadowOffset.X) { DefaultStyleCDO->bUsesDropShadow ? ShadowOffset.X = DefaultStyleCDO->ShadowOffset.X : 0.f; } else { bAllDefaults = false; }
-			if (ShadowOffset.Y == CDO->ShadowOffset.Y) { DefaultStyleCDO->bUsesDropShadow ? ShadowOffset.Y = DefaultStyleCDO->ShadowOffset.Y : 0.f; } else { bAllDefaults = false; }
+			{
+				bool bShadowOffsetHasChanged = false;
+				FVector2D MyShadowOffset = GetShadowOffset();
+				const FVector2D CDOShadowOffset = CDO->GetShadowOffset();
+				if (MyShadowOffset.X == CDOShadowOffset.X)
+				{
+					DefaultStyleCDO->bUsesDropShadow ? MyShadowOffset.X = DefaultStyleCDO->ShadowOffset.X : 0.f;
+					bShadowOffsetHasChanged = true;
+				}
+				if (MyShadowOffset.Y == CDOShadowOffset.Y)
+				{
+					DefaultStyleCDO->bUsesDropShadow ? MyShadowOffset.Y = DefaultStyleCDO->ShadowOffset.Y : 0.f;
+					bShadowOffsetHasChanged = true;
+				}
+
+				if (bShadowOffsetHasChanged)
+				{
+					SetShadowOffset(MyShadowOffset);
+				}
+				bAllDefaults = bAllDefaults && !bShadowOffsetHasChanged;
+			}
 
 			// ShadowColorAndOpacity
-			if (ShadowColorAndOpacity.R == CDO->ShadowColorAndOpacity.R) { DefaultStyleCDO->bUsesDropShadow ? ShadowColorAndOpacity.R = DefaultStyleCDO->ShadowColor.R : 0.f; } else { bAllDefaults = false; }
-			if (ShadowColorAndOpacity.G == CDO->ShadowColorAndOpacity.G) { DefaultStyleCDO->bUsesDropShadow ? ShadowColorAndOpacity.G = DefaultStyleCDO->ShadowColor.G : 0.f; } else { bAllDefaults = false; }
-			if (ShadowColorAndOpacity.B == CDO->ShadowColorAndOpacity.B) { DefaultStyleCDO->bUsesDropShadow ? ShadowColorAndOpacity.B = DefaultStyleCDO->ShadowColor.B : 0.f; } else { bAllDefaults = false; }
-			if (ShadowColorAndOpacity.A == CDO->ShadowColorAndOpacity.A) { DefaultStyleCDO->bUsesDropShadow ? ShadowColorAndOpacity.A = DefaultStyleCDO->ShadowColor.A : 0.f; } else { bAllDefaults = false; }
+			{
+				bool bShadowColorAndOpacityHasChanged = false;
+				FLinearColor MyShadowOffset = GetShadowColorAndOpacity();
+				const FLinearColor CDOShadowOffset = CDO->GetShadowColorAndOpacity();
+				if (MyShadowOffset.R == CDOShadowOffset.R)
+				{
+					DefaultStyleCDO->bUsesDropShadow ? MyShadowOffset.R = DefaultStyleCDO->ShadowColor.R : 0.f;
+					bShadowColorAndOpacityHasChanged = true;
+				}
+				if (MyShadowOffset.G == CDOShadowOffset.G)
+				{
+					DefaultStyleCDO->bUsesDropShadow ? MyShadowOffset.G = DefaultStyleCDO->ShadowColor.G : 0.f;
+					bShadowColorAndOpacityHasChanged = true;
+				}
+				if (MyShadowOffset.B == CDOShadowOffset.B)
+				{
+					DefaultStyleCDO->bUsesDropShadow ? MyShadowOffset.B = DefaultStyleCDO->ShadowColor.B : 0.f;
+					bShadowColorAndOpacityHasChanged = true;
+				}
+				if (MyShadowOffset.A == CDOShadowOffset.A)
+				{
+					DefaultStyleCDO->bUsesDropShadow ? MyShadowOffset.A = DefaultStyleCDO->ShadowColor.A : 0.f;
+					bShadowColorAndOpacityHasChanged = true;
+				}
+
+				if (bShadowColorAndOpacityHasChanged)
+				{
+					SetShadowColorAndOpacity(MyShadowOffset);
+				}
+				bAllDefaults = bAllDefaults && !bShadowColorAndOpacityHasChanged;
+			}
 
 			if (bAllDefaults)
 			{
@@ -368,6 +473,7 @@ void UCommonTextBlock::Serialize(FArchive& Ar)
 		FVector2D TempShadowOffset = FVector2D::ZeroVector;
 		FLinearColor TempShadowColorAndOpacity = FLinearColor::Transparent;
 
+		PRAGMA_DISABLE_DEPRECATION_WARNINGS
 		Swap(Font, TempFont);
 		Swap(StrikeBrush, TempStrikeBrush);
 		Swap(Margin, TempMargin);
@@ -385,6 +491,7 @@ void UCommonTextBlock::Serialize(FArchive& Ar)
 		Swap(TempColorAndOpacity, ColorAndOpacity);
 		Swap(TempShadowOffset, ShadowOffset);
 		Swap(TempShadowColorAndOpacity, ShadowColorAndOpacity);
+		PRAGMA_ENABLE_DEPRECATION_WARNINGS
 	}
 	else
 #endif //if WITH_EDITORONLY_DATA
@@ -393,7 +500,7 @@ void UCommonTextBlock::Serialize(FArchive& Ar)
 
 		if (Ar.IsLoading() && bDisplayAllCaps_DEPRECATED)
 		{
-			TextTransformPolicy = ETextTransformPolicy::ToUpper;
+			SetTextTransformPolicy(ETextTransformPolicy::ToUpper);
 			bDisplayAllCaps_DEPRECATED = false;
 		}
 	}
@@ -407,7 +514,7 @@ void UCommonTextBlock::SetWrapTextWidth(int32 InWrapTextAt)
 
 void UCommonTextBlock::SetTextCase(bool bUseAllCaps)
 {
-	TextTransformPolicy = bUseAllCaps ? ETextTransformPolicy::ToUpper : ETextTransformPolicy::None;
+	SetTextTransformPolicy(bUseAllCaps ? ETextTransformPolicy::ToUpper : ETextTransformPolicy::None);
 	SynchronizeProperties();
 }
 
@@ -437,7 +544,7 @@ void UCommonTextBlock::SynchronizeProperties()
 
 	if (bAutoCollapseWithEmptyText)
 	{
-		SetVisibility(Text.IsEmpty() ? ESlateVisibility::Collapsed : ESlateVisibility::SelfHitTestInvisible);
+		SetVisibility(GetText().IsEmpty() ? ESlateVisibility::Collapsed : ESlateVisibility::SelfHitTestInvisible);
 	}
 
 	Super::SynchronizeProperties();
@@ -485,21 +592,21 @@ void UCommonTextBlock::UpdateFromStyle()
 {
 	if (const UCommonTextStyle* TextStyle = GetStyleCDO())
 	{
-		Font = TextStyle->Font;
-		StrikeBrush = TextStyle->StrikeBrush;
+		SetFont(TextStyle->Font);
+		SetStrikeBrush(TextStyle->StrikeBrush);
 		Margin = TextStyle->Margin;
 		LineHeightPercentage = TextStyle->LineHeightPercentage;
-		ColorAndOpacity = TextStyle->Color;
+		SetColorAndOpacity(TextStyle->Color);
 
 		if (TextStyle->bUsesDropShadow)
 		{
-			ShadowOffset = TextStyle->ShadowOffset;
-			ShadowColorAndOpacity = TextStyle->ShadowColor;
+			SetShadowOffset(TextStyle->ShadowOffset);
+			SetShadowColorAndOpacity(TextStyle->ShadowColor);
 		}
 		else
 		{
-			ShadowOffset = FVector2D::ZeroVector;
-			ShadowColorAndOpacity = FLinearColor::Transparent;
+			SetShadowOffset(FVector2D::ZeroVector);
+			SetShadowColorAndOpacity(FLinearColor::Transparent);
 		}
 	}
 }
@@ -534,6 +641,7 @@ bool UCommonTextBlock::CanEditChange(const FProperty* InProperty) const
 	{
 		if (const UCommonTextStyle* TextStyle = GetStyleCDO())
 		{
+PRAGMA_DISABLE_DEPRECATION_WARNINGS
 			static TArray<FName> InvalidPropertiesWithStyle =
 			{
 				GET_MEMBER_NAME_CHECKED(UCommonTextBlock, Font),
@@ -544,7 +652,7 @@ bool UCommonTextBlock::CanEditChange(const FProperty* InProperty) const
 				GET_MEMBER_NAME_CHECKED(UCommonTextBlock, ShadowOffset),
 				GET_MEMBER_NAME_CHECKED(UCommonTextBlock, ShadowColorAndOpacity)
 			};
-
+PRAGMA_ENABLE_DEPRECATION_WARNINGS
 			return !InvalidPropertiesWithStyle.Contains(InProperty->GetFName());
 		}
 
@@ -609,7 +717,7 @@ void UCommonTextBlock::ApplyFontSizeMultiplier() const
 {
 	if (MyTextBlock.IsValid())
 	{
-		FSlateFontInfo FontInfo = Font;
+		FSlateFontInfo FontInfo = GetFont();
 		FontInfo.Size *= MobileFontSizeMultiplier;
 		MyTextBlock->SetFont(FontInfo);
 	}
diff --git a/Engine/Source/Runtime/UMG/Private/Components/TextBlock.cpp b/Engine/Source/Runtime/UMG/Private/Components/TextBlock.cpp
--- a/Engine/Source/Runtime/UMG/Private/Components/TextBlock.cpp
+++ b/Engine/Source/Runtime/UMG/Private/Components/TextBlock.cpp
@@ -20,6 +20,7 @@ UTextBlock::UTextBlock(const FObjectInitializer& ObjectInitializer)
 {
 	bIsVariable = false;
 	bWrapWithInvalidationPanel = false;
+PRAGMA_DISABLE_DEPRECATION_WARNINGS
 	ShadowOffset = FVector2D(1.0f, 1.0f);
 	ColorAndOpacity = FLinearColor::White;
 	ShadowColorAndOpacity = FLinearColor::Transparent;
@@ -31,6 +32,7 @@ UTextBlock::UTextBlock(const FObjectInitializer& ObjectInitializer)
 		static ConstructorHelpers::FObjectFinder<UFont> RobotoFontObj(*UWidget::GetDefaultFontName());
 		Font = FSlateFontInfo(RobotoFontObj.Object, 24, FName("Bold"));
 	}
+PRAGMA_ENABLE_DEPRECATION_WARNINGS
 
 #if WITH_EDITORONLY_DATA
 	AccessibleBehavior = ESlateAccessibleBehavior::Auto;
@@ -45,6 +47,19 @@ void UTextBlock::ReleaseSlateResources(bool bReleaseChildren)
 	MyTextBlock.Reset();
 }
 
+PRAGMA_DISABLE_DEPRECATION_WARNINGS
+FSlateColor UTextBlock::GetColorAndOpacity() const
+{
+	if (ColorAndOpacityDelegate.IsBound() && !IsDesignTime())
+	{
+		return ColorAndOpacityDelegate.Execute();
+	}
+	else
+	{
+		return ColorAndOpacity;
+	}
+}
+
 void UTextBlock::SetColorAndOpacity(FSlateColor InColorAndOpacity)
 {
 	ColorAndOpacity = InColorAndOpacity;
@@ -62,6 +77,18 @@ void UTextBlock::SetOpacity(float InOpacity)
 	SetColorAndOpacity(FSlateColor(CurrentColor));
 }
 
+FLinearColor UTextBlock::GetShadowColorAndOpacity() const
+{
+	if (ShadowColorAndOpacityDelegate.IsBound() && !IsDesignTime())
+	{
+		return ShadowColorAndOpacityDelegate.Execute();
+	}
+	else
+	{
+		return ShadowColorAndOpacity;
+	}
+}
+
 void UTextBlock::SetShadowColorAndOpacity(FLinearColor InShadowColorAndOpacity)
 {
 	ShadowColorAndOpacity = InShadowColorAndOpacity;
@@ -71,6 +98,11 @@ void UTextBlock::SetShadowColorAndOpacity(FLinearColor InShadowColorAndOpacity)
 	}
 }
 
+FVector2D UTextBlock::GetShadowOffset() const
+{
+	return ShadowOffset;
+}
+
 void UTextBlock::SetShadowOffset(FVector2D InShadowOffset)
 {
 	ShadowOffset = InShadowOffset;
@@ -80,6 +112,11 @@ void UTextBlock::SetShadowOffset(FVector2D InShadowOffset)
 	}
 }
 
+const FSlateFontInfo& UTextBlock::GetFont() const
+{
+	return Font;
+}
+
 void UTextBlock::SetFont(FSlateFontInfo InFontInfo)
 {
 	Font = InFontInfo;
@@ -89,6 +126,11 @@ void UTextBlock::SetFont(FSlateFontInfo InFontInfo)
 	}
 }
 
+const FSlateBrush& UTextBlock::GetStrikeBrush() const
+{
+	return StrikeBrush;
+}
+
 void UTextBlock::SetStrikeBrush(FSlateBrush InStrikeBrush)
 {
 	StrikeBrush = InStrikeBrush;
@@ -107,6 +149,11 @@ void UTextBlock::SetJustification(ETextJustify::Type InJustification)
 	}
 }
 
+float UTextBlock::GetMinDesiredWidth() const
+{
+	return MinDesiredWidth;
+}
+
 void UTextBlock::SetMinDesiredWidth(float InMinDesiredWidth)
 {
 	MinDesiredWidth = InMinDesiredWidth;
@@ -125,6 +172,11 @@ void UTextBlock::SetAutoWrapText(bool InAutoWrapText)
 	}
 }
 
+ETextTransformPolicy UTextBlock::GetTextTransformPolicy() const
+{
+	return TextTransformPolicy;
+}
+
 void UTextBlock::SetTextTransformPolicy(ETextTransformPolicy InTransformPolicy)
 {
 	TextTransformPolicy = InTransformPolicy;
@@ -134,11 +186,17 @@ void UTextBlock::SetTextTransformPolicy(ETextTransformPolicy InTransformPolicy)
 	}
 }
 
+ETextOverflowPolicy UTextBlock::GetTextOverflowPolicy() const
+{
+	return TextOverflowPolicy;
+}
+
 void UTextBlock::SetTextOverflowPolicy(ETextOverflowPolicy InOverflowPolicy)
 {
 	TextOverflowPolicy = InOverflowPolicy;
 	SynchronizeProperties();
 }
+PRAGMA_DISABLE_DEPRECATION_WARNINGS
 
 UMaterialInstanceDynamic* UTextBlock::GetDynamicFontMaterial()
 {
@@ -184,6 +242,7 @@ UMaterialInstanceDynamic* UTextBlock::GetDynamicOutlineMaterial()
 	return nullptr;
 }
 
+
 TSharedRef<SWidget> UTextBlock::RebuildWidget()
 {
  	if (bWrapWithInvalidationPanel && !IsDesignTime())
diff --git a/Engine/Source/Runtime/UMG/Public/Components/TextBlock.h b/Engine/Source/Runtime/UMG/Public/Components/TextBlock.h
--- a/Engine/Source/Runtime/UMG/Public/Components/TextBlock.h
+++ b/Engine/Source/Runtime/UMG/Public/Components/TextBlock.h
@@ -24,6 +24,86 @@
 	GENERATED_UCLASS_BODY()
 
 public:
+	/** The text to display */
+	UE_DEPRECATED(5.1, "Direct access to Text is deprecated. Please use the getter or setter.")
+	UPROPERTY(EditAnywhere, BlueprintReadWrite, Getter, Setter, BlueprintGetter="GetText", BlueprintSetter="SetText", Category="Content", meta = (MultiLine = "true"))
+	FText Text;
+	
+	/** A bindable delegate to allow logic to drive the text of the widget */
+	UPROPERTY()
+	FGetText TextDelegate;
+
+	/** The color of the text */
+	UE_DEPRECATED(5.1, "Direct access to ColorAndOpacity is deprecated. Please use the getter or setter.")
+	UPROPERTY(EditAnywhere, BlueprintReadWrite, Getter, Setter, BlueprintSetter="SetColorAndOpacity", Category="Appearance")
+	FSlateColor ColorAndOpacity;
+
+	/** A bindable delegate for the ColorAndOpacity. */
+	UPROPERTY()
+	FGetSlateColor ColorAndOpacityDelegate;
+	
+	/** The font to render the text with */
+	UE_DEPRECATED(5.1, "Direct access to Font is deprecated. Please use the getter or setter.")
+	UPROPERTY(EditAnywhere, BlueprintReadWrite, Getter, Setter, BlueprintSetter="SetFont", Category="Appearance")
+	FSlateFontInfo Font;
+
+	/** The brush to strike through text with */
+	UE_DEPRECATED(5.1, "Direct access to StrikeBrush is deprecated. Please use the getter or setter.")
+	UPROPERTY(EditAnywhere, BlueprintReadWrite, Getter, Setter, BlueprintSetter="SetStrikeBrush", Category="Appearance")
+	FSlateBrush StrikeBrush;
+
+	/** The direction the shadow is cast */
+	UE_DEPRECATED(5.1, "Direct access to ShadowOffset is deprecated. Please use the getter or setter.")
+	UPROPERTY(EditAnywhere, BlueprintReadWrite, Getter, Setter, BlueprintSetter="SetShadowOffset", Category="Appearance")
+	FVector2D ShadowOffset;
+
+	/** The color of the shadow */
+	UE_DEPRECATED(5.1, "Direct access to ShadowColorAndOpacity is deprecated. Please use the getter or setter.")
+	UPROPERTY(EditAnywhere, BlueprintReadWrite, Getter, Setter, BlueprintSetter="SetShadowColorAndOpacity", Category="Appearance", meta=( DisplayName="Shadow Color" ))
+	FLinearColor ShadowColorAndOpacity;
+
+	/** A bindable delegate for the ShadowColorAndOpacity. */
+	UPROPERTY()
+	FGetLinearColor ShadowColorAndOpacityDelegate;
+
+	/** The minimum desired size for the text */
+	UE_DEPRECATED(5.1, "Direct access to MinDesiredWidth is deprecated. Please use the getter or setter.")
+	UPROPERTY(EditAnywhere, BlueprintReadWrite, Getter, Setter, BlueprintSetter="SetMinDesiredWidth", Category="Appearance")
+	float MinDesiredWidth;
+
+	/** If true, it will automatically wrap this text widget with an invalidation panel */
+	UPROPERTY(EditAnywhere, BlueprintReadOnly, Category="Performance", AdvancedDisplay)
+	bool bWrapWithInvalidationPanel;
+
+	/** The text transformation policy to apply to this text block. */
+	UE_DEPRECATED(5.1, "Direct access to TextTransformPolicy is deprecated. Please use the getter or setter.")
+	UPROPERTY(EditAnywhere, BlueprintReadWrite, Getter, Setter, BlueprintSetter="SetTextTransformPolicy", Category="Appearance", meta=(DisplayName="Transform Policy"))
+	ETextTransformPolicy TextTransformPolicy;
+
+	/** Sets what happens to text that is clipped and doesn't fit within the clip rect for this widget */
+	UE_DEPRECATED(5.1, "Direct access to TextOverflowPolicy is deprecated. Please use the getter or setter.")
+	UPROPERTY(EditAnywhere, BlueprintReadWrite, Getter, Setter, BlueprintSetter="SetTextOverflowPolicy", Category="Clipping", AdvancedDisplay, meta = (DisplayName="Overflow Policy"))
+	ETextOverflowPolicy TextOverflowPolicy;
+
+public:
+	/** 
+	 * Gets the widget text
+	 * @return The widget text
+	 */
+	UFUNCTION(BlueprintCallable, Category="Widget", meta=(DisplayName="GetText (Text)"))
+	FText GetText() const;
+
+	/**
+	 * Directly sets the widget text.
+	 * Warning: This will wipe any binding created for the Text property!
+	 * @param InText The text to assign to the widget
+	 */
+	UFUNCTION(BlueprintCallable, Category="Widget", meta=(DisplayName="SetText (Text)"))
+	virtual void SetText(FText InText);
+
+	/** @return the color and opacity of the text in this text block. */
+	FSlateColor GetColorAndOpacity() const;
+
 	/**  
 	 * Sets the color and opacity of the text in this text block
 	 *
@@ -37,9 +117,12 @@
 	*
 	* @param InOpacity		The new text opacity
 	*/
-	UFUNCTION(BlueprintCallable, Category = "Appearance")
+	UFUNCTION(BlueprintCallable, Category="Appearance")
 	void SetOpacity(float InOpacity);
 
+	/** @return the color and opacity of the text drop shadow. */
+	FLinearColor GetShadowColorAndOpacity() const;
+
 	/**  
 	 * Sets the color and opacity of the text drop shadow
 	 * Note: if opacity is zero no shadow will be drawn
@@ -49,6 +132,9 @@
 	UFUNCTION(BlueprintCallable, Category="Appearance")
 	void SetShadowColorAndOpacity(FLinearColor InShadowColorAndOpacity);
 
+	/** @return the offset that the text drop shadow should be drawn at. */
+	FVector2D GetShadowOffset() const;
+
 	/**  
 	 * Sets the offset that the text drop shadow should be drawn at
 	 *
@@ -57,28 +143,37 @@
 	UFUNCTION(BlueprintCallable, Category="Appearance")
 	void SetShadowOffset(FVector2D InShadowOffset);
 
+	/** @return the font info for this text block. */
+	const FSlateFontInfo& GetFont() const;
+
 	/**
 	 * Dynamically set the font info for this text block
 	 * 
 	 * @param InFontInfo The new font info
 	 */
-	UFUNCTION(BlueprintCallable, Category = "Appearance")
+	UFUNCTION(BlueprintCallable, Category="Appearance")
 	void SetFont(FSlateFontInfo InFontInfo);
 
+	/** @return the brush to use to strike through text.  */
+	const FSlateBrush& GetStrikeBrush() const;
+
 	/**
 	 * Dynamically set the strike brush for this text block
 	 *
 	 * @param InStrikeBrush The new brush to use to strike through text
 	 */
-	UFUNCTION(BlueprintCallable, Category = "Appearance")
+	UFUNCTION(BlueprintCallable, Category="Appearance")
 	void SetStrikeBrush(FSlateBrush InStrikeBrush);
 
+	/** @return the minimum desired width for this text block */
+	float GetMinDesiredWidth() const;
+
 	/**
 	 *  Set the minimum desired width for this text block
 	 *
 	 *  @param InMinDesiredWidth new minimum desired width
 	 */
-	UFUNCTION(BlueprintCallable, Category = "Appearance")
+	UFUNCTION(BlueprintCallable, Category="Appearance")
 	void SetMinDesiredWidth(float InMinDesiredWidth);
 	
 	/**
@@ -86,23 +181,29 @@
 	 *
 	 * @param InAutoTextWrap to turn wrap on or off.
 	 */
-	UFUNCTION(BlueprintCallable, Category = "Appearance")
+	UFUNCTION(BlueprintCallable, Category="Appearance")
 	void SetAutoWrapText(bool InAutoTextWrap);
 
+	/** @return the text transformation policy for this text block. */
+	ETextTransformPolicy GetTextTransformPolicy() const;
+
 	/**
 	 * Set the text transformation policy for this text block.
 	 *
 	 * @param InTransformPolicy the new text transformation policy.
 	 */
-	UFUNCTION(BlueprintCallable, Category = "Appearance")
+	UFUNCTION(BlueprintCallable, Category="Appearance")
 	void SetTextTransformPolicy(ETextTransformPolicy InTransformPolicy);
 
+	/** @return the text overflow policy for this text block. */
+	ETextOverflowPolicy GetTextOverflowPolicy() const;
+
 	/**
 	 * Set the text overflow policy for this text block.
 	 *
 	 * @param InOverflowPolicy the new text overflow policy.
 	 */
-	UFUNCTION(BlueprintCallable, Category = "Appearance")
+	UFUNCTION(BlueprintCallable, Category="Appearance")
 	void SetTextOverflowPolicy(ETextOverflowPolicy InOverflowPolicy);
 
 	/**  */
@@ -114,73 +215,6 @@
 	UMaterialInstanceDynamic* GetDynamicOutlineMaterial();
 
 public:
-	/** The text to display */
-	UPROPERTY(EditAnywhere, Category=Content, meta=( MultiLine="true" ))
-	FText Text;
-	
-	/** A bindable delegate to allow logic to drive the text of the widget */
-	UPROPERTY()
-	FGetText TextDelegate;
-
-	/** The color of the text */
-	UPROPERTY(EditAnywhere, BlueprintReadOnly, Category=Appearance)
-	FSlateColor ColorAndOpacity;
-
-	/** A bindable delegate for the ColorAndOpacity. */
-	UPROPERTY()
-	FGetSlateColor ColorAndOpacityDelegate;
-	
-	/** The font to render the text with */
-	UPROPERTY(EditAnywhere, BlueprintReadOnly, Category=Appearance)
-	FSlateFontInfo Font;
-
-	/** The brush to strike through text with */
-	UPROPERTY(EditAnywhere, BlueprintReadOnly, Category = Appearance)
-	FSlateBrush StrikeBrush;
-
-	/** The direction the shadow is cast */
-	UPROPERTY(EditAnywhere, BlueprintReadOnly, Category=Appearance)
-	FVector2D ShadowOffset;
-
-	/** The color of the shadow */
-	UPROPERTY(EditAnywhere, BlueprintReadOnly, Category=Appearance, meta=( DisplayName="Shadow Color" ))
-	FLinearColor ShadowColorAndOpacity;
-
-	/** A bindable delegate for the ShadowColorAndOpacity. */
-	UPROPERTY()
-	FGetLinearColor ShadowColorAndOpacityDelegate;
-
-	/** The minimum desired size for the text */
-	UPROPERTY(EditAnywhere, BlueprintReadOnly, Category=Appearance)
-	float MinDesiredWidth;
-
-	/** If true, it will automatically wrap this text widget with an invalidation panel */
-	UPROPERTY(EditAnywhere, BlueprintReadOnly, Category = Performance, AdvancedDisplay)
-	bool bWrapWithInvalidationPanel;
-
-	/** The text transformation policy to apply to this text block */
-	UPROPERTY(EditAnywhere, BlueprintReadOnly, Category=Appearance, meta=(DisplayName="Transform Policy"))
-	ETextTransformPolicy TextTransformPolicy;
-
-	/** Sets what happens to text that is clipped and doesn't fit within the clip rect for this widget */
-	UPROPERTY(EditAnywhere, BlueprintReadWrite, Category = Clipping, AdvancedDisplay, meta = (DisplayName = "Overflow Policy"))
-	ETextOverflowPolicy TextOverflowPolicy;
-
-	/** 
-	 * Gets the widget text
-	 * @return The widget text
-	 */
-	UFUNCTION(BlueprintCallable, Category="Widget", meta=(DisplayName="GetText (Text)"))
-	FText GetText() const;
-
-	/**
-	 * Directly sets the widget text.
-	 * Warning: This will wipe any binding created for the Text property!
-	 * @param InText The text to assign to the widget
-	 */
-	UFUNCTION(BlueprintCallable, Category="Widget", meta=(DisplayName="SetText (Text)"))
-	virtual void SetText(FText InText);
-
 	//~ Begin UTextLayoutWidget Interface
 	virtual void SetJustification(ETextJustify::Type InJustification) override;
 	//~ End UTextLayoutWidget Interface
