UCommonActivatableWidgetSwitcher:
- If the switcher has an owning activatable widget, only activate switcher child if the owning widget is active.
- Removed unnecessary includes.

diff --git a/Engine/Plugins/Experimental/CommonUI/Source/CommonUI/Private/CommonActivatableWidgetSwitcher.cpp b/Engine/Plugins/Experimental/CommonUI/Source/CommonUI/Private/CommonActivatableWidgetSwitcher.cpp
--- a/Engine/Plugins/Experimental/CommonUI/Source/CommonUI/Private/CommonActivatableWidgetSwitcher.cpp
+++ b/Engine/Plugins/Experimental/CommonUI/Source/CommonUI/Private/CommonActivatableWidgetSwitcher.cpp
@@ -1,25 +1,107 @@
 // Copyright Epic Games, Inc. All Rights Reserved.
 
 #include "CommonActivatableWidgetSwitcher.h"
-#include "CommonWidgetPaletteCategories.h"
-#include "Components/WidgetSwitcherSlot.h"
-#include "Components/PanelSlot.h"
+
+#include "CommonActivatableWidget.h"
+#include "CommonUIUtils.h"
+
+void UCommonActivatableWidgetSwitcher::ReleaseSlateResources(bool bReleaseChildren)
+{
+	Super::ReleaseSlateResources(bReleaseChildren);
+
+	UnbindOwningActivatableWidget();
+}
+
+void UCommonActivatableWidgetSwitcher::OnWidgetRebuilt()
+{
+	Super::OnWidgetRebuilt();
+
+	BindOwningActivatableWidget();
+}
 
 void UCommonActivatableWidgetSwitcher::HandleOutgoingWidget()
 {
-	if (UCommonActivatableWidget* OutgoingActivatable = Cast<UCommonActivatableWidget>(GetWidgetAtIndex(ActiveWidgetIndex)))
-	{
-		OutgoingActivatable->DeactivateWidget();
-	}
+	DeactivateActiveWidget();
 }
 
 void UCommonActivatableWidgetSwitcher::HandleSlateActiveIndexChanged(int32 ActiveIndex)
 {
 	Super::HandleSlateActiveIndexChanged(ActiveIndex);
 
-	if (UCommonActivatableWidget* IncomingActivatable = Cast<UCommonActivatableWidget>(GetWidgetAtIndex(ActiveWidgetIndex)))
+	if (!WeakOwningActivatableWidget.IsSet())
+	{
+		BindOwningActivatableWidget();
+	}
+
+	AttemptToActivateActiveWidget();
+}
+
+void UCommonActivatableWidgetSwitcher::HandleOwningWidgetActivationChanged(const bool bIsActivated)
+{
+	if (bIsActivated)
+	{
+		AttemptToActivateActiveWidget();
+	}
+	else
+	{
+		DeactivateActiveWidget();
+	}
+}
+
+void UCommonActivatableWidgetSwitcher::AttemptToActivateActiveWidget()
+{
+	if (UCommonActivatableWidget* ActivatableWidget = Cast<UCommonActivatableWidget>(GetActiveWidget()))
+	{
+		const UCommonActivatableWidget* OwningActivatableWidget = GetOwningActivatableWidget();
+
+		// Only activate if our owning activatable widget is activated, or if we don't have an owning activatable widget
+		if (OwningActivatableWidget == nullptr || OwningActivatableWidget->IsActivated())
+		{
+			ActivatableWidget->ActivateWidget();
+		}
+	}
+}
+
+void UCommonActivatableWidgetSwitcher::DeactivateActiveWidget()
+{
+	if (UCommonActivatableWidget* ActivatableWidget = Cast<UCommonActivatableWidget>(GetActiveWidget()))
 	{
-		IncomingActivatable->ActivateWidget();
+		ActivatableWidget->DeactivateWidget();
 	}
 }
 
+void UCommonActivatableWidgetSwitcher::BindOwningActivatableWidget()
+{
+	UCommonActivatableWidget* OwningActivatableWidget = CommonUIUtils::GetOwningUserWidget<UCommonActivatableWidget>(this);
+
+	if (GetOwningActivatableWidget() != OwningActivatableWidget)
+	{
+		UnbindOwningActivatableWidget();
+
+		if (OwningActivatableWidget)
+		{
+			OwningActivatableWidget->OnActivated().AddUObject(this, &UCommonActivatableWidgetSwitcher::HandleOwningWidgetActivationChanged, true);
+			OwningActivatableWidget->OnDeactivated().AddUObject(this, &UCommonActivatableWidgetSwitcher::HandleOwningWidgetActivationChanged, false);
+
+			HandleOwningWidgetActivationChanged(OwningActivatableWidget->IsActivated());
+		}
+	}
+	
+	WeakOwningActivatableWidget = OwningActivatableWidget;
+}
+
+void UCommonActivatableWidgetSwitcher::UnbindOwningActivatableWidget()
+{
+	if (UCommonActivatableWidget* OwningActivatableWidget = GetOwningActivatableWidget())
+	{
+		OwningActivatableWidget->OnActivated().RemoveAll(this);
+		OwningActivatableWidget->OnDeactivated().RemoveAll(this);
+	}
+
+	WeakOwningActivatableWidget.Reset();
+}
+
+UCommonActivatableWidget* UCommonActivatableWidgetSwitcher::GetOwningActivatableWidget() const
+{
+	return WeakOwningActivatableWidget.IsSet() ? WeakOwningActivatableWidget.GetValue().Get() : nullptr;
+}
diff --git a/Engine/Plugins/Experimental/CommonUI/Source/CommonUI/Public/CommonActivatableWidgetSwitcher.h b/Engine/Plugins/Experimental/CommonUI/Source/CommonUI/Public/CommonActivatableWidgetSwitcher.h
--- a/Engine/Plugins/Experimental/CommonUI/Source/CommonUI/Public/CommonActivatableWidgetSwitcher.h
+++ b/Engine/Plugins/Experimental/CommonUI/Source/CommonUI/Public/CommonActivatableWidgetSwitcher.h
@@ -2,14 +2,15 @@
 
 #pragma once
 
-#include "Components/WidgetSwitcher.h"
-#include "CommonActivatableWidget.h"
-#include "Slate/SCommonAnimatedSwitcher.h"
 #include "CommonAnimatedSwitcher.h"
-#include "Components/Widget.h"
+
+#include "Misc/Optional.h"
+#include "UObject/WeakObjectPtrTemplates.h"
 
 #include "CommonActivatableWidgetSwitcher.generated.h"
 
+class UCommonActivatableWidget;
+
 /**
  * An animated switcher that knows about CommonActivatableWidgets. It can also hold other Widgets.
  */
@@ -18,7 +19,25 @@ class COMMONUI_API UCommonActivatableWidgetSwitcher : public UCommonAnimatedSwit
 {
 	GENERATED_BODY()
 
+public:
+	virtual void ReleaseSlateResources(bool bReleaseChildren) override;
+
 protected:
+	virtual void OnWidgetRebuilt() override;
+
 	virtual void HandleOutgoingWidget() override;
 	virtual void HandleSlateActiveIndexChanged(int32 ActiveIndex) override;
-};
\ No newline at end of file
+
+private:
+	void HandleOwningWidgetActivationChanged(const bool bIsActivated);
+
+	void AttemptToActivateActiveWidget();
+	void DeactivateActiveWidget();
+
+	void BindOwningActivatableWidget();
+	void UnbindOwningActivatableWidget();
+
+	UCommonActivatableWidget* GetOwningActivatableWidget() const;
+
+	TOptional<TWeakObjectPtr<UCommonActivatableWidget>> WeakOwningActivatableWidget;
+};
