Checking validity of RootButton in UCommonButtonBase::RefreshDimensions to avoid crashes when the preview widget is rebuilt
See UDN: https://epicgames.lightning.force.com/lightning/r/Case/5004z00001dHIK7AAO/view

diff --git a/Engine/Plugins/Experimental/CommonUI/Source/CommonUI/Private/CommonButtonBase.cpp b/Engine/Plugins/Experimental/CommonUI/Source/CommonUI/Private/CommonButtonBase.cpp
--- a/Engine/Plugins/Experimental/CommonUI/Source/CommonUI/Private/CommonButtonBase.cpp
+++ b/Engine/Plugins/Experimental/CommonUI/Source/CommonUI/Private/CommonButtonBase.cpp
@@ -746,9 +746,12 @@ void UCommonButtonBase::SetSelectedInternal(bool bInSelected, bool bAllowSound /
 
 void UCommonButtonBase::RefreshDimensions()
 {
-	const UCommonButtonStyle* const StyleCDO = GetStyleCDO();
-	RootButton->SetMinDesiredWidth(FMath::Max(MinWidth, StyleCDO ? StyleCDO->MinWidth : 0));
-	RootButton->SetMinDesiredHeight(FMath::Max(MinHeight, StyleCDO ? StyleCDO->MinHeight : 0));
+	if (RootButton.IsValid())
+	{
+		const UCommonButtonStyle* const StyleCDO = GetStyleCDO();
+		RootButton->SetMinDesiredWidth(FMath::Max(MinWidth, StyleCDO ? StyleCDO->MinWidth : 0));
+		RootButton->SetMinDesiredHeight(FMath::Max(MinHeight, StyleCDO ? StyleCDO->MinHeight : 0));
+	}
 }
 
 void UCommonButtonBase::NativeOnMouseEnter(const FGeometry& InGeometry, const FPointerEvent& InMouseEvent)
