Remove option to disable action domains from CommonUI

diff --git a/Engine/Plugins/Experimental/CommonUI/Source/CommonUI/Private/Input/CommonUIActionRouterBase.cpp b/Engine/Plugins/Experimental/CommonUI/Source/CommonUI/Private/Input/CommonUIActionRouterBase.cpp
--- a/Engine/Plugins/Experimental/CommonUI/Source/CommonUI/Private/Input/CommonUIActionRouterBase.cpp
+++ b/Engine/Plugins/Experimental/CommonUI/Source/CommonUI/Private/Input/CommonUIActionRouterBase.cpp
@@ -31,12 +31,6 @@ static const FAutoConsoleVariableRef CVarAlwaysShowCursor(
 	bAlwaysShowCursor,
 	TEXT(""));
 
-bool bEnableActionDomainRouting = true;
-static const FAutoConsoleVariableRef CVarEnableActionDomainRouting(
-	TEXT("CommonUI.EnableActionDomainRouting"),
-	bEnableActionDomainRouting,
-	TEXT("Enables the routing of UI inputs based on the action domain instead of the active root node."));
-
 //@todo DanH: TEMP LOCATION
 FGlobalUITags FGlobalUITags::GUITags;
 
@@ -360,7 +354,7 @@ ERouteUIInputResult UCommonUIActionRouterBase::ProcessInput(FKey Key, EInputEven
 			ProcessHoldResult = ActiveRootNode->ProcessHoldInput(ActiveMode, Key, InputEvent);
 		}
 
-		if (bEnableActionDomainRouting && ProcessHoldResult == EProcessHoldActionResult::Unhandled)
+		if (ProcessHoldResult == EProcessHoldActionResult::Unhandled)
 		{
 			ProcessHoldResult = ProcessHoldInputOnActionDomains(ActiveMode, Key, InputEvent);
 		}
@@ -377,7 +371,7 @@ ERouteUIInputResult UCommonUIActionRouterBase::ProcessInput(FKey Key, EInputEven
 					bHandled = ActiveRootNode->ProcessNormalInput(ActiveMode, Key, Event);
 				}
 
-				if (bEnableActionDomainRouting && !bHandled)
+				if (!bHandled)
 				{
 					bHandled = ProcessInputOnActionDomains(ActiveMode, Key, Event);
 				}
@@ -588,18 +582,20 @@ void UCommonUIActionRouterBase::HandleRootWidgetSlateReleased(TWeakPtr<FActivata
 			ActiveInputConfig.Reset();
 		}
 	}
-	else if (bEnableActionDomainRouting)
+	else
 	{
+		int32 NumRemoved = 0;
 		for (TPair<TObjectPtr<UCommonInputActionDomain>, FActionDomainSortedRootList>& Pair : ActionDomainRootNodes)
 		{
 			FActionDomainSortedRootList& ActionDomainRootList = Pair.Value;
-			ActionDomainRootList.Remove(RootNode);
+			NumRemoved += ActionDomainRootList.Remove(RootNode);
+		}
+
+		if (NumRemoved <= 0)
+		{
+			UCommonActivatableWidget* ActivatableWidget = RootNode->GetWidget();
+			ensureAlwaysMsgf(false, TEXT("Root node could not be found during deactivation [%s]"), ActivatableWidget ? *ActivatableWidget->GetName() : TEXT("Unknown"));
 		}
-	}
-	else
-	{
-		UCommonActivatableWidget* ActivatableWidget = RootNode->GetWidget();
-		ensureAlwaysMsgf(false, TEXT("Root node could not be found during deactivation [%s]"), ActivatableWidget ? *ActivatableWidget->GetName() : TEXT("Unknown"));
 	}
 }
 
@@ -781,7 +777,7 @@ void UCommonUIActionRouterBase::ProcessRebuiltWidgets()
 	{
 		if (RebuiltWidget.IsValid() && RebuiltWidget->GetCachedWidget())
 		{
-			if (UCommonInputActionDomain* ActionDomain = bEnableActionDomainRouting ? RebuiltWidget->GetCalculatedActionDomain() : nullptr)
+			if (UCommonInputActionDomain* ActionDomain = RebuiltWidget->GetCalculatedActionDomain())
 			{
 				FActivatableTreeRootRef RootNode = FActivatableTreeRoot::Create(*this, *RebuiltWidget);
 				RebuiltWidget->OnSlateReleased().AddUObject(this, &UCommonUIActionRouterBase::HandleRootWidgetSlateReleased, TWeakPtr<FActivatableTreeRoot>(RootNode));
@@ -930,10 +926,7 @@ bool UCommonUIActionRouterBase::Tick(float DeltaTime)
 		{
 			if (ActiveRootNode->ProcessHoldInput(ActiveMode, HeldKey, EInputEvent::IE_Repeat) == EProcessHoldActionResult::Unhandled)
 			{
-				if (bEnableActionDomainRouting)
-				{
-					ProcessHoldInputOnActionDomains(ActiveMode, HeldKey, EInputEvent::IE_Repeat);
-				}
+				ProcessHoldInputOnActionDomains(ActiveMode, HeldKey, EInputEvent::IE_Repeat);
 			}
 		}
 	}
@@ -1124,7 +1117,7 @@ FActivatableTreeNodePtr UCommonUIActionRouterBase::FindNode(const UCommonActivat
 			}
 		}
 
-		if (bEnableActionDomainRouting && !FoundNode.IsValid())
+		if (!FoundNode.IsValid())
 		{
 			for (const TPair<TObjectPtr<UCommonInputActionDomain>, FActionDomainSortedRootList>& Pair : ActionDomainRootNodes)
 			{
