Quick fix to prevent recursive call in CommonButton.
#preflgiht none

diff --git a/Engine/Plugins/Experimental/CommonUI/Source/CommonUI/Private/CommonButtonBase.cpp b/Engine/Plugins/Experimental/CommonUI/Source/CommonUI/Private/CommonButtonBase.cpp
--- a/Engine/Plugins/Experimental/CommonUI/Source/CommonUI/Private/CommonButtonBase.cpp
+++ b/Engine/Plugins/Experimental/CommonUI/Source/CommonUI/Private/CommonButtonBase.cpp
@@ -434,7 +434,7 @@ void UCommonButtonBase::SetIsEnabled(bool bInIsEnabled)
 	else
 	{
 		// Change the underlying enabled bool but do not call the case because we don't want to propogate it to the underlying SWidget
-		SetIsEnabled(bInIsEnabled);
+		Super::SetIsEnabled(bInIsEnabled);
 		DisableButton();
 	}
 }
