UMG: Update UWidget::Visibility initialization from SetVisibility to SetVisibilityInternal.

diff --git a/Engine/Plugins/Experimental/CommonUI/Source/CommonUI/Private/CommonBorder.cpp b/Engine/Plugins/Experimental/CommonUI/Source/CommonUI/Private/CommonBorder.cpp
--- a/Engine/Plugins/Experimental/CommonUI/Source/CommonUI/Private/CommonBorder.cpp
+++ b/Engine/Plugins/Experimental/CommonUI/Source/CommonUI/Private/CommonBorder.cpp
@@ -26,7 +26,7 @@ UCommonBorder::UCommonBorder(const FObjectInitializer& ObjectInitializer)
 {
 	Background.DrawAs = ESlateBrushDrawType::NoDrawType;
 
-	SetVisibility(ESlateVisibility::SelfHitTestInvisible);
+	SetVisibilityInternal(ESlateVisibility::SelfHitTestInvisible);
 }
 
 void UCommonBorder::PostLoad()
diff --git a/Engine/Plugins/Experimental/CommonUI/Source/CommonUI/Private/CommonLazyWidget.cpp b/Engine/Plugins/Experimental/CommonUI/Source/CommonUI/Private/CommonLazyWidget.cpp
--- a/Engine/Plugins/Experimental/CommonUI/Source/CommonUI/Private/CommonLazyWidget.cpp
+++ b/Engine/Plugins/Experimental/CommonUI/Source/CommonUI/Private/CommonLazyWidget.cpp
@@ -9,7 +9,7 @@ UCommonLazyWidget::UCommonLazyWidget(const FObjectInitializer& Initializer)
 	: Super(Initializer)
 {
 	LoadingBackgroundBrush.DrawAs = ESlateBrushDrawType::NoDrawType;
-	SetVisibility(ESlateVisibility::SelfHitTestInvisible);
+	SetVisibilityInternal(ESlateVisibility::SelfHitTestInvisible);
 }
 
 TSharedRef<SWidget> UCommonLazyWidget::RebuildWidget()
diff --git a/Engine/Plugins/Experimental/CommonUI/Source/CommonUI/Private/CommonListView.cpp b/Engine/Plugins/Experimental/CommonUI/Source/CommonUI/Private/CommonListView.cpp
--- a/Engine/Plugins/Experimental/CommonUI/Source/CommonUI/Private/CommonListView.cpp
+++ b/Engine/Plugins/Experimental/CommonUI/Source/CommonUI/Private/CommonListView.cpp
@@ -12,7 +12,7 @@
 UCommonListView::UCommonListView(const FObjectInitializer& ObjectInitializer)
 	: Super(ObjectInitializer)
 {
-	SetVisibility(ESlateVisibility::Visible);
+	SetVisibilityInternal(ESlateVisibility::Visible);
 	bEnableScrollAnimation = true;
 }
 
diff --git a/Engine/Plugins/Experimental/CommonUI/Source/CommonUI/Private/CommonLoadGuard.cpp b/Engine/Plugins/Experimental/CommonUI/Source/CommonUI/Private/CommonLoadGuard.cpp
--- a/Engine/Plugins/Experimental/CommonUI/Source/CommonUI/Private/CommonLoadGuard.cpp
+++ b/Engine/Plugins/Experimental/CommonUI/Source/CommonUI/Private/CommonLoadGuard.cpp
@@ -271,7 +271,7 @@ UCommonLoadGuard::UCommonLoadGuard(const FObjectInitializer& Initializer)
 	, ThrobberAlignment(EHorizontalAlignment::HAlign_Center)
 	, LoadingText(NSLOCTEXT("LoadGuard", "Loading", "Loading..."))
 {
-	SetVisibility(ESlateVisibility::SelfHitTestInvisible);
+	SetVisibilityInternal(ESlateVisibility::SelfHitTestInvisible);
 
 	// Default to not showing the loading BG brush - it's an opt-in kind of thing
 	LoadingBackgroundBrush.DrawAs = ESlateBrushDrawType::NoDrawType;
diff --git a/Engine/Plugins/Experimental/CommonUI/Source/CommonUI/Private/CommonTextBlock.cpp b/Engine/Plugins/Experimental/CommonUI/Source/CommonUI/Private/CommonTextBlock.cpp
--- a/Engine/Plugins/Experimental/CommonUI/Source/CommonUI/Private/CommonTextBlock.cpp
+++ b/Engine/Plugins/Experimental/CommonUI/Source/CommonUI/Private/CommonTextBlock.cpp
@@ -276,7 +276,7 @@ void UCommonTextStyle::ApplyToTextBlock(const TSharedRef<STextBlock>& TextBlock)
 UCommonTextBlock::UCommonTextBlock(const FObjectInitializer& ObjectInitializer)
 	: Super(ObjectInitializer)
 {
-	SetVisibility(ESlateVisibility::SelfHitTestInvisible);
+	SetVisibilityInternal(ESlateVisibility::SelfHitTestInvisible);
 	SetClipping(EWidgetClipping::Inherit);
 }
 
diff --git a/Engine/Plugins/Experimental/CommonUI/Source/CommonUI/Private/Widgets/CommonActivatableWidgetContainer.cpp b/Engine/Plugins/Experimental/CommonUI/Source/CommonUI/Private/Widgets/CommonActivatableWidgetContainer.cpp
--- a/Engine/Plugins/Experimental/CommonUI/Source/CommonUI/Private/Widgets/CommonActivatableWidgetContainer.cpp
+++ b/Engine/Plugins/Experimental/CommonUI/Source/CommonUI/Private/Widgets/CommonActivatableWidgetContainer.cpp
@@ -32,7 +32,7 @@ UCommonActivatableWidgetContainerBase::UCommonActivatableWidgetContainerBase(con
 	: Super(Initializer)
 	, GeneratedWidgetsPool(*this)
 {
-	SetVisibility(ESlateVisibility::Collapsed);
+	SetVisibilityInternal(ESlateVisibility::Collapsed);
 }
 
 void UCommonActivatableWidgetContainerBase::AddWidgetInstance(UCommonActivatableWidget& ActivatableWidget)
diff --git a/Engine/Source/Runtime/UMG/Private/Components/BackgroundBlur.cpp b/Engine/Source/Runtime/UMG/Private/Components/BackgroundBlur.cpp
--- a/Engine/Source/Runtime/UMG/Private/Components/BackgroundBlur.cpp
+++ b/Engine/Source/Runtime/UMG/Private/Components/BackgroundBlur.cpp
@@ -25,7 +25,7 @@ UBackgroundBlur::UBackgroundBlur(const FObjectInitializer& ObjectInitializer)
 	, LowQualityFallbackBrush(FSlateNoResource())
 {
 	bIsVariable = false;
-	SetVisibility(ESlateVisibility::SelfHitTestInvisible);
+	SetVisibilityInternal(ESlateVisibility::SelfHitTestInvisible);
 }
 
 void UBackgroundBlur::ReleaseSlateResources(bool bReleaseChildren)
diff --git a/Engine/Source/Runtime/UMG/Private/Components/CanvasPanel.cpp b/Engine/Source/Runtime/UMG/Private/Components/CanvasPanel.cpp
--- a/Engine/Source/Runtime/UMG/Private/Components/CanvasPanel.cpp
+++ b/Engine/Source/Runtime/UMG/Private/Components/CanvasPanel.cpp
@@ -13,7 +13,7 @@ UCanvasPanel::UCanvasPanel(const FObjectInitializer& ObjectInitializer)
 	: Super(ObjectInitializer)
 {
 	bIsVariable = false;
-	SetVisibility(ESlateVisibility::SelfHitTestInvisible);
+	SetVisibilityInternal(ESlateVisibility::SelfHitTestInvisible);
 }
 
 void UCanvasPanel::ReleaseSlateResources(bool bReleaseChildren)
diff --git a/Engine/Source/Runtime/UMG/Private/Components/DynamicEntryBoxBase.cpp b/Engine/Source/Runtime/UMG/Private/Components/DynamicEntryBoxBase.cpp
--- a/Engine/Source/Runtime/UMG/Private/Components/DynamicEntryBoxBase.cpp
+++ b/Engine/Source/Runtime/UMG/Private/Components/DynamicEntryBoxBase.cpp
@@ -15,7 +15,7 @@ UDynamicEntryBoxBase::UDynamicEntryBoxBase(const FObjectInitializer& Initializer
 	: Super(Initializer)
 	, EntryWidgetPool(*this)
 {
-	SetVisibility(ESlateVisibility::SelfHitTestInvisible);
+	SetVisibilityInternal(ESlateVisibility::SelfHitTestInvisible);
 	EntrySizeRule.SizeRule = ESlateSizeRule::Automatic;
 }
 
diff --git a/Engine/Source/Runtime/UMG/Private/Components/GridPanel.cpp b/Engine/Source/Runtime/UMG/Private/Components/GridPanel.cpp
--- a/Engine/Source/Runtime/UMG/Private/Components/GridPanel.cpp
+++ b/Engine/Source/Runtime/UMG/Private/Components/GridPanel.cpp
@@ -14,7 +14,7 @@ UGridPanel::UGridPanel(const FObjectInitializer& ObjectInitializer)
 	: Super(ObjectInitializer)
 {
 	bIsVariable = false;
-	SetVisibility(ESlateVisibility::SelfHitTestInvisible);
+	SetVisibilityInternal(ESlateVisibility::SelfHitTestInvisible);
 }
 
 void UGridPanel::ReleaseSlateResources(bool bReleaseChildren)
diff --git a/Engine/Source/Runtime/UMG/Private/Components/HorizontalBox.cpp b/Engine/Source/Runtime/UMG/Private/Components/HorizontalBox.cpp
--- a/Engine/Source/Runtime/UMG/Private/Components/HorizontalBox.cpp
+++ b/Engine/Source/Runtime/UMG/Private/Components/HorizontalBox.cpp
@@ -12,7 +12,7 @@ UHorizontalBox::UHorizontalBox(const FObjectInitializer& ObjectInitializer)
 	: Super(ObjectInitializer)
 {
 	bIsVariable = false;
-	SetVisibility(ESlateVisibility::SelfHitTestInvisible);
+	SetVisibilityInternal(ESlateVisibility::SelfHitTestInvisible);
 }
 
 void UHorizontalBox::ReleaseSlateResources(bool bReleaseChildren)
diff --git a/Engine/Source/Runtime/UMG/Private/Components/InvalidationBox.cpp b/Engine/Source/Runtime/UMG/Private/Components/InvalidationBox.cpp
--- a/Engine/Source/Runtime/UMG/Private/Components/InvalidationBox.cpp
+++ b/Engine/Source/Runtime/UMG/Private/Components/InvalidationBox.cpp
@@ -14,7 +14,7 @@ UInvalidationBox::UInvalidationBox(const FObjectInitializer& ObjectInitializer)
 	: Super(ObjectInitializer)
 {
 	bCanCache = true;
-	SetVisibility(ESlateVisibility::SelfHitTestInvisible);
+	SetVisibilityInternal(ESlateVisibility::SelfHitTestInvisible);
 }
 
 void UInvalidationBox::ReleaseSlateResources(bool bReleaseChildren)
diff --git a/Engine/Source/Runtime/UMG/Private/Components/NamedSlot.cpp b/Engine/Source/Runtime/UMG/Private/Components/NamedSlot.cpp
--- a/Engine/Source/Runtime/UMG/Private/Components/NamedSlot.cpp
+++ b/Engine/Source/Runtime/UMG/Private/Components/NamedSlot.cpp
@@ -15,7 +15,7 @@ UNamedSlot::UNamedSlot(const FObjectInitializer& ObjectInitializer)
 	: Super(ObjectInitializer)
 {
 	bIsVariable = true;
-	SetVisibility(ESlateVisibility::SelfHitTestInvisible);
+	SetVisibilityInternal(ESlateVisibility::SelfHitTestInvisible);
 }
 
 void UNamedSlot::ReleaseSlateResources(bool bReleaseChildren)
diff --git a/Engine/Source/Runtime/UMG/Private/Components/Overlay.cpp b/Engine/Source/Runtime/UMG/Private/Components/Overlay.cpp
--- a/Engine/Source/Runtime/UMG/Private/Components/Overlay.cpp
+++ b/Engine/Source/Runtime/UMG/Private/Components/Overlay.cpp
@@ -12,7 +12,7 @@ UOverlay::UOverlay(const FObjectInitializer& ObjectInitializer)
 	: Super(ObjectInitializer)
 {
 	bIsVariable = false;
-	SetVisibility(ESlateVisibility::SelfHitTestInvisible);
+	SetVisibilityInternal(ESlateVisibility::SelfHitTestInvisible);
 }
 
 void UOverlay::ReleaseSlateResources(bool bReleaseChildren)
diff --git a/Engine/Source/Runtime/UMG/Private/Components/RetainerBox.cpp b/Engine/Source/Runtime/UMG/Private/Components/RetainerBox.cpp
--- a/Engine/Source/Runtime/UMG/Private/Components/RetainerBox.cpp
+++ b/Engine/Source/Runtime/UMG/Private/Components/RetainerBox.cpp
@@ -16,7 +16,7 @@ static FName DefaultTextureParameterName("Texture");
 URetainerBox::URetainerBox(const FObjectInitializer& ObjectInitializer)
 	: Super(ObjectInitializer)
 {
-	SetVisibility(ESlateVisibility::Visible);
+	SetVisibilityInternal(ESlateVisibility::Visible);
 	Phase = 0;
 	PhaseCount = 1;
 	RenderOnPhase = true;
diff --git a/Engine/Source/Runtime/UMG/Private/Components/RichTextBlock.cpp b/Engine/Source/Runtime/UMG/Private/Components/RichTextBlock.cpp
--- a/Engine/Source/Runtime/UMG/Private/Components/RichTextBlock.cpp
+++ b/Engine/Source/Runtime/UMG/Private/Components/RichTextBlock.cpp
@@ -47,7 +47,7 @@ FORCEINLINE TSharedPtr< ObjectType > MakeShareableDeferredCleanup(ObjectType* In
 URichTextBlock::URichTextBlock(const FObjectInitializer& ObjectInitializer)
 	: Super(ObjectInitializer)
 {
-	SetVisibility(ESlateVisibility::SelfHitTestInvisible);
+	SetVisibilityInternal(ESlateVisibility::SelfHitTestInvisible);
 	TextTransformPolicy = ETextTransformPolicy::None;
 	TextOverflowPolicy = ETextOverflowPolicy::Clip;
 }
diff --git a/Engine/Source/Runtime/UMG/Private/Components/SafeZone.cpp b/Engine/Source/Runtime/UMG/Private/Components/SafeZone.cpp
--- a/Engine/Source/Runtime/UMG/Private/Components/SafeZone.cpp
+++ b/Engine/Source/Runtime/UMG/Private/Components/SafeZone.cpp
@@ -14,7 +14,7 @@ USafeZone::USafeZone()
 	, PadBottom(true)
 {
 	bCanHaveMultipleChildren = false;
-	SetVisibility(ESlateVisibility::SelfHitTestInvisible);
+	SetVisibilityInternal(ESlateVisibility::SelfHitTestInvisible);
 }
 
 #if WITH_EDITOR
diff --git a/Engine/Source/Runtime/UMG/Private/Components/ScaleBox.cpp b/Engine/Source/Runtime/UMG/Private/Components/ScaleBox.cpp
--- a/Engine/Source/Runtime/UMG/Private/Components/ScaleBox.cpp
+++ b/Engine/Source/Runtime/UMG/Private/Components/ScaleBox.cpp
@@ -13,7 +13,7 @@ UScaleBox::UScaleBox(const FObjectInitializer& ObjectInitializer)
 	: Super(ObjectInitializer)
 {
 	bIsVariable = false;
-	SetVisibility(ESlateVisibility::SelfHitTestInvisible);
+	SetVisibilityInternal(ESlateVisibility::SelfHitTestInvisible);
 
 	StretchDirection = EStretchDirection::Both;
 	Stretch = EStretch::ScaleToFit;
diff --git a/Engine/Source/Runtime/UMG/Private/Components/ScrollBox.cpp b/Engine/Source/Runtime/UMG/Private/Components/ScrollBox.cpp
--- a/Engine/Source/Runtime/UMG/Private/Components/ScrollBox.cpp
+++ b/Engine/Source/Runtime/UMG/Private/Components/ScrollBox.cpp
@@ -37,7 +37,7 @@ UScrollBox::UScrollBox(const FObjectInitializer& ObjectInitializer)
 {
 	bIsVariable = false;
 
-	SetVisibility(ESlateVisibility::Visible);
+	SetVisibilityInternal(ESlateVisibility::Visible);
 	SetClipping(EWidgetClipping::ClipToBounds);
 
 	if (DefaultScrollBoxStyle == nullptr)
diff --git a/Engine/Source/Runtime/UMG/Private/Components/SizeBox.cpp b/Engine/Source/Runtime/UMG/Private/Components/SizeBox.cpp
--- a/Engine/Source/Runtime/UMG/Private/Components/SizeBox.cpp
+++ b/Engine/Source/Runtime/UMG/Private/Components/SizeBox.cpp
@@ -15,7 +15,7 @@ USizeBox::USizeBox(const FObjectInitializer& ObjectInitializer)
 	: Super(ObjectInitializer)
 {
 	bIsVariable = false;
-	SetVisibility(ESlateVisibility::SelfHitTestInvisible);
+	SetVisibilityInternal(ESlateVisibility::SelfHitTestInvisible);
 	MinAspectRatio = 1;
 	MaxAspectRatio = 1;
 }
diff --git a/Engine/Source/Runtime/UMG/Private/Components/Spacer.cpp b/Engine/Source/Runtime/UMG/Private/Components/Spacer.cpp
--- a/Engine/Source/Runtime/UMG/Private/Components/Spacer.cpp
+++ b/Engine/Source/Runtime/UMG/Private/Components/Spacer.cpp
@@ -14,7 +14,7 @@ USpacer::USpacer(const FObjectInitializer& ObjectInitializer)
 	, Size(1.0f, 1.0f)
 {
 	bIsVariable = false;
-	SetVisibility(ESlateVisibility::SelfHitTestInvisible);
+	SetVisibilityInternal(ESlateVisibility::SelfHitTestInvisible);
 }
 
 void USpacer::ReleaseSlateResources(bool bReleaseChildren)
diff --git a/Engine/Source/Runtime/UMG/Private/Components/UniformGridPanel.cpp b/Engine/Source/Runtime/UMG/Private/Components/UniformGridPanel.cpp
--- a/Engine/Source/Runtime/UMG/Private/Components/UniformGridPanel.cpp
+++ b/Engine/Source/Runtime/UMG/Private/Components/UniformGridPanel.cpp
@@ -14,7 +14,7 @@ UUniformGridPanel::UUniformGridPanel(const FObjectInitializer& ObjectInitializer
 	: Super(ObjectInitializer)
 {
 	bIsVariable = false;
-	SetVisibility(ESlateVisibility::SelfHitTestInvisible);
+	SetVisibilityInternal(ESlateVisibility::SelfHitTestInvisible);
 }
 
 void UUniformGridPanel::ReleaseSlateResources(bool bReleaseChildren)
diff --git a/Engine/Source/Runtime/UMG/Private/Components/VerticalBox.cpp b/Engine/Source/Runtime/UMG/Private/Components/VerticalBox.cpp
--- a/Engine/Source/Runtime/UMG/Private/Components/VerticalBox.cpp
+++ b/Engine/Source/Runtime/UMG/Private/Components/VerticalBox.cpp
@@ -12,7 +12,7 @@ UVerticalBox::UVerticalBox(const FObjectInitializer& ObjectInitializer)
 	: Super(ObjectInitializer)
 {
 	bIsVariable = false;
-	SetVisibility(ESlateVisibility::SelfHitTestInvisible);
+	SetVisibilityInternal(ESlateVisibility::SelfHitTestInvisible);
 }
 
 void UVerticalBox::ReleaseSlateResources(bool bReleaseChildren)
diff --git a/Engine/Source/Runtime/UMG/Private/Components/Widget.cpp b/Engine/Source/Runtime/UMG/Private/Components/Widget.cpp
--- a/Engine/Source/Runtime/UMG/Private/Components/Widget.cpp
+++ b/Engine/Source/Runtime/UMG/Private/Components/Widget.cpp
@@ -375,6 +375,11 @@ ESlateVisibility UWidget::GetVisibility() const
 }
 
 void UWidget::SetVisibility(ESlateVisibility InVisibility)
+{
+	SetVisibilityInternal(InVisibility);
+}
+
+void UWidget::SetVisibilityInternal(ESlateVisibility InVisibility)
 {
 	if (Visibility != InVisibility)
 	{
diff --git a/Engine/Source/Runtime/UMG/Private/Components/WidgetSwitcher.cpp b/Engine/Source/Runtime/UMG/Private/Components/WidgetSwitcher.cpp
--- a/Engine/Source/Runtime/UMG/Private/Components/WidgetSwitcher.cpp
+++ b/Engine/Source/Runtime/UMG/Private/Components/WidgetSwitcher.cpp
@@ -13,7 +13,7 @@ UWidgetSwitcher::UWidgetSwitcher(const FObjectInitializer& ObjectInitializer)
 	: Super(ObjectInitializer)
 {
 	bIsVariable = true;
-	SetVisibility(ESlateVisibility::SelfHitTestInvisible);
+	SetVisibilityInternal(ESlateVisibility::SelfHitTestInvisible);
 }
 
 void UWidgetSwitcher::ReleaseSlateResources(bool bReleaseChildren)
diff --git a/Engine/Source/Runtime/UMG/Private/Components/WindowTitleBarArea.cpp b/Engine/Source/Runtime/UMG/Private/Components/WindowTitleBarArea.cpp
--- a/Engine/Source/Runtime/UMG/Private/Components/WindowTitleBarArea.cpp
+++ b/Engine/Source/Runtime/UMG/Private/Components/WindowTitleBarArea.cpp
@@ -19,7 +19,7 @@ UWindowTitleBarArea::UWindowTitleBarArea(const FObjectInitializer& ObjectInitial
 	: Super(ObjectInitializer)
 {
 	bIsVariable = false;
-	SetVisibility(ESlateVisibility::Visible);
+	SetVisibilityInternal(ESlateVisibility::Visible);
 
 	bDoubleClickTogglesFullscreen = false;
 }
diff --git a/Engine/Source/Runtime/UMG/Private/Components/WrapBox.cpp b/Engine/Source/Runtime/UMG/Private/Components/WrapBox.cpp
--- a/Engine/Source/Runtime/UMG/Private/Components/WrapBox.cpp
+++ b/Engine/Source/Runtime/UMG/Private/Components/WrapBox.cpp
@@ -12,7 +12,7 @@ UWrapBox::UWrapBox(const FObjectInitializer& ObjectInitializer)
 	: Super(ObjectInitializer)
 {
 	bIsVariable = false;
-	SetVisibility(ESlateVisibility::SelfHitTestInvisible);
+	SetVisibilityInternal(ESlateVisibility::SelfHitTestInvisible);
 	WrapSize = 500;
 	bExplicitWrapSize = false;
 	HorizontalAlignment = HAlign_Left;
diff --git a/Engine/Source/Runtime/UMG/Private/UserWidget.cpp b/Engine/Source/Runtime/UMG/Private/UserWidget.cpp
--- a/Engine/Source/Runtime/UMG/Private/UserWidget.cpp
+++ b/Engine/Source/Runtime/UMG/Private/UserWidget.cpp
@@ -77,7 +77,7 @@ UUserWidget::UUserWidget(const FObjectInitializer& ObjectInitializer)
 	, TickFrequency(EWidgetTickFrequency::Auto)
 {
 	ViewportAnchors = FAnchors(0, 0, 1, 1);
-	SetVisibility(ESlateVisibility::SelfHitTestInvisible);
+	SetVisibilityInternal(ESlateVisibility::SelfHitTestInvisible);
 
 	bIsFocusable = false;
 	ColorAndOpacity = FLinearColor::White;
diff --git a/Engine/Source/Runtime/UMG/Public/Components/Widget.h b/Engine/Source/Runtime/UMG/Public/Components/Widget.h
--- a/Engine/Source/Runtime/UMG/Public/Components/Widget.h
+++ b/Engine/Source/Runtime/UMG/Public/Components/Widget.h
@@ -555,6 +555,10 @@ public:
 	UFUNCTION(BlueprintCallable, Category="Widget")
 	virtual void SetVisibility(ESlateVisibility InVisibility);
 
+protected:
+	void SetVisibilityInternal(ESlateVisibility InVisibility);
+
+public:
 	/** Gets the current visibility of the widget. */
 	UFUNCTION(BlueprintCallable, Category="Widget")
 	float GetRenderOpacity() const;
