Speculative Fix ensure CommonUI postload

diff --git a/Engine/Plugins/Experimental/CommonUI/Source/CommonUI/Private/CommonBorder.cpp b/Engine/Plugins/Experimental/CommonUI/Source/CommonUI/Private/CommonBorder.cpp
--- a/Engine/Plugins/Experimental/CommonUI/Source/CommonUI/Private/CommonBorder.cpp
+++ b/Engine/Plugins/Experimental/CommonUI/Source/CommonUI/Private/CommonBorder.cpp
@@ -40,7 +40,9 @@ void UCommonBorder::PostLoad()
 		UCommonBorder* CDO = Cast<UCommonBorder>(GetClass()->GetDefaultObject());
 		if (Background == CDO->Background)
 		{
-			Style = ICommonUIModule::GetEditorSettings().GetTemplateBorderStyle();
+			UCommonUIEditorSettings& Settings = ICommonUIModule::GetEditorSettings();
+			Settings.ConditionalPostLoad();
+			Style = Settings.GetTemplateBorderStyle();
 		}
 	}
 
diff --git a/Engine/Plugins/Experimental/CommonUI/Source/CommonUI/Private/CommonButtonBase.cpp b/Engine/Plugins/Experimental/CommonUI/Source/CommonUI/Private/CommonButtonBase.cpp
--- a/Engine/Plugins/Experimental/CommonUI/Source/CommonUI/Private/CommonButtonBase.cpp
+++ b/Engine/Plugins/Experimental/CommonUI/Source/CommonUI/Private/CommonButtonBase.cpp
@@ -335,7 +335,9 @@ void UCommonButtonBase::PostLoad()
 	// We will remove this once existing content is fixed up. Since previously the native CDO was actually the default style, this code will attempt to set the style on assets that were once using this default
 	if (!Style && !bStyleNoLongerNeedsConversion && !IsRunningDedicatedServer())
 	{
-		Style = ICommonUIModule::GetEditorSettings().GetTemplateButtonStyle();
+		UCommonUIEditorSettings& Settings = ICommonUIModule::GetEditorSettings();
+		Settings.ConditionalPostLoad();
+		Style = Settings.GetTemplateButtonStyle();
 	}
 
 	bStyleNoLongerNeedsConversion = true;
diff --git a/Engine/Plugins/Experimental/CommonUI/Source/CommonUI/Private/CommonLoadGuard.cpp b/Engine/Plugins/Experimental/CommonUI/Source/CommonUI/Private/CommonLoadGuard.cpp
--- a/Engine/Plugins/Experimental/CommonUI/Source/CommonUI/Private/CommonLoadGuard.cpp
+++ b/Engine/Plugins/Experimental/CommonUI/Source/CommonUI/Private/CommonLoadGuard.cpp
@@ -351,7 +351,9 @@ void UCommonLoadGuard::PostLoad()
 	//TODO: I want this removed, it's needed for backwards compatibility
 	if (!TextStyle && !bStyleNoLongerNeedsConversion && !IsRunningDedicatedServer())
 	{
-		TextStyle = ICommonUIModule::GetEditorSettings().GetTemplateTextStyle();
+		UCommonUIEditorSettings& Settings = ICommonUIModule::GetEditorSettings();
+		Settings.ConditionalPostLoad();
+		TextStyle = Settings.GetTemplateTextStyle();
 	}
 
 	bStyleNoLongerNeedsConversion = true;
diff --git a/Engine/Plugins/Experimental/CommonUI/Source/CommonUI/Private/CommonTextBlock.cpp b/Engine/Plugins/Experimental/CommonUI/Source/CommonUI/Private/CommonTextBlock.cpp
--- a/Engine/Plugins/Experimental/CommonUI/Source/CommonUI/Private/CommonTextBlock.cpp
+++ b/Engine/Plugins/Experimental/CommonUI/Source/CommonUI/Private/CommonTextBlock.cpp
@@ -295,10 +295,13 @@ void UCommonTextBlock::PostLoad()
 	// We will remove this once existing content is fixed up. Since previously the native CDO was actually the default style, this code will attempt to set the style on assets that were once using this default
 	if (!Style && !bStyleNoLongerNeedsConversion && !IsRunningDedicatedServer())
 	{
-		TSubclassOf<UCommonTextStyle> DefaultStyle = ICommonUIModule::GetEditorSettings().GetTemplateTextStyle();
+		UCommonUIEditorSettings& Settings = ICommonUIModule::GetEditorSettings();
+		Settings.ConditionalPostLoad();
+		TSubclassOf<UCommonTextStyle> DefaultStyle = Settings.GetTemplateTextStyle();
 		UCommonTextStyle* DefaultStyleCDO = DefaultStyle ? Cast<UCommonTextStyle>(DefaultStyle->ClassDefaultObject) : nullptr;
 		if (DefaultStyleCDO)
 		{
+			DefaultStyleCDO->ConditionalPostLoad();
 			UCommonTextBlock* CDO = CastChecked<UCommonTextBlock>(GetClass()->GetDefaultObject());
 			bool bAllDefaults = true;
 
@@ -445,7 +448,7 @@ void UCommonTextBlock::PostLoad()
 
 			if (bAllDefaults)
 			{
-				Style = ICommonUIModule::GetEditorSettings().GetTemplateTextStyle();
+				Style = Settings.GetTemplateTextStyle();
 			}
 		}
 	}
