Adds BP callable setter for bHideInputAction on CommonButtonBase.

diff --git a/Engine/Plugins/Experimental/CommonUI/Source/CommonUI/Private/CommonButtonBase.cpp b/Engine/Plugins/Experimental/CommonUI/Source/CommonUI/Private/CommonButtonBase.cpp
--- a/Engine/Plugins/Experimental/CommonUI/Source/CommonUI/Private/CommonButtonBase.cpp
+++ b/Engine/Plugins/Experimental/CommonUI/Source/CommonUI/Private/CommonButtonBase.cpp
@@ -571,6 +571,13 @@ void UCommonButtonBase::SetIsInteractionEnabled(bool bInIsInteractionEnabled)
 	}
 }
 
+void UCommonButtonBase::SetHideInputAction(bool bInHideInputAction)
+{
+	bHideInputAction = bInHideInputAction;
+
+	UpdateInputActionWidgetVisibility();
+}
+
 bool UCommonButtonBase::IsInteractionEnabled() const
 {
 	ESlateVisibility Vis = GetVisibility(); // hidden or collapsed should have 'bInteractionEnabled' set false, but sometimes they don't :(
diff --git a/Engine/Plugins/Experimental/CommonUI/Source/CommonUI/Public/CommonButtonBase.h b/Engine/Plugins/Experimental/CommonUI/Source/CommonUI/Public/CommonButtonBase.h
--- a/Engine/Plugins/Experimental/CommonUI/Source/CommonUI/Public/CommonButtonBase.h
+++ b/Engine/Plugins/Experimental/CommonUI/Source/CommonUI/Public/CommonButtonBase.h
@@ -300,6 +300,9 @@ public:
 	UFUNCTION(BlueprintCallable, Category = "Common Button|Setters")
 	void SetIsInteractionEnabled(bool bInIsInteractionEnabled);
 
+	UFUNCTION(BlueprintCallable, Category = "Common Button|Setters")
+	void SetHideInputAction(bool bInHideInputAction);
+
 	/** Is this button currently interactable? (use instead of GetIsEnabled) */
 	UFUNCTION(BlueprintCallable, Category = "Common Button|Getters")
 	bool IsInteractionEnabled() const;
