Remove platform references from CommonUI

diff --git a/Engine/Plugins/Experimental/CommonUI/Source/CommonInput/CommonInput.Build.cs b/Engine/Plugins/Experimental/CommonUI/Source/CommonInput/CommonInput.Build.cs
--- a/Engine/Plugins/Experimental/CommonUI/Source/CommonInput/CommonInput.Build.cs
+++ b/Engine/Plugins/Experimental/CommonUI/Source/CommonInput/CommonInput.Build.cs
@@ -51,6 +51,9 @@ public class CommonInput : ModuleRules
 		}
 
 		PrivateDependencyModuleNames.Add("GeForceNOWWrapper");
+
+		PrivateDefinitions.Add("UE_COMMONINPUT_PLATFORM_KBM_REQUIRES_ATTACHED_MOUSE=" + (bPlatformKBMRequiresAttachedMouse ? "1" : "0"));
+		PrivateDefinitions.Add("UE_COMMONINPUT_PLATFORM_SUPPORTS_TOUCH=" + (bPlatformSupportsTouch ? "1" : "0"));
 	}
 
 	static public string ToCommonUIPlatform(UnrealTargetPlatform TargetPlatform)
@@ -78,4 +81,14 @@ public class CommonInput : ModuleRules
 
 		return string.Empty;
 	}
+
+	protected virtual bool bPlatformKBMRequiresAttachedMouse { get { return false; } }
+	protected virtual bool bPlatformSupportsTouch
+	{
+		get
+		{
+			// Support touch testing until touch is supported on desktop
+			return Target.Platform.IsInGroup(UnrealPlatformGroup.Desktop) && Target.Configuration != UnrealTargetConfiguration.Shipping;
+		}
+	}
 }
diff --git a/Engine/Plugins/Experimental/CommonUI/Source/CommonInput/Private/CommonInputSubsystem.cpp b/Engine/Plugins/Experimental/CommonUI/Source/CommonInput/Private/CommonInputSubsystem.cpp
--- a/Engine/Plugins/Experimental/CommonUI/Source/CommonInput/Private/CommonInputSubsystem.cpp
+++ b/Engine/Plugins/Experimental/CommonUI/Source/CommonInput/Private/CommonInputSubsystem.cpp
@@ -719,18 +719,14 @@ bool UCommonInputSubsystem::PlatformSupportsInputType(ECommonInputType InInputTy
 	{
 		case ECommonInputType::MouseAndKeyboard:
 		{
-#if PLATFORM_SWITCH
+#if UE_COMMONINPUT_PLATFORM_KBM_REQUIRES_ATTACHED_MOUSE
 			bPlatformSupportsInput &= FSlateApplication::Get().IsMouseAttached();
 #endif
 		}
 		break;
 		case ECommonInputType::Touch:
 		{
-#if PLATFORM_SWITCH
-			bPlatformSupportsInput &= SUPPORT_SWITCH_TOUCHSCREEN;
-#elif PLATFORM_DESKTOP && !UE_BUILD_SHIPPING
-			bPlatformSupportsInput = true; // Support touch testing until touch is supported on desktop
-#endif
+			bPlatformSupportsInput &= UE_COMMONINPUT_PLATFORM_SUPPORTS_TOUCH != 0;
 		}
 		break;
 		case ECommonInputType::Gamepad:
diff --git a/Engine/Plugins/Experimental/CommonUI/Source/CommonUI/CommonUI.Build.cs b/Engine/Plugins/Experimental/CommonUI/Source/CommonUI/CommonUI.Build.cs
--- a/Engine/Plugins/Experimental/CommonUI/Source/CommonUI/CommonUI.Build.cs
+++ b/Engine/Plugins/Experimental/CommonUI/Source/CommonUI/CommonUI.Build.cs
@@ -56,5 +56,17 @@ public class CommonUI : ModuleRules
                 }
             );
         }
-    }
+		PrivateDefinitions.Add("UE_COMMONUI_PLATFORM_KBM_REQUIRES_ATTACHED_MOUSE=" + (bPlatformKBMRequiresAttachedMouse ? "1" : "0"));
+		PrivateDefinitions.Add("UE_COMMONUI_PLATFORM_SUPPORTS_TOUCH=" + (bPlatformSupportsTouch ? "1" : "0"));
+	}
+
+	protected virtual bool bPlatformKBMRequiresAttachedMouse { get { return false; } }
+	protected virtual bool bPlatformSupportsTouch
+	{
+		get
+		{
+			// Support touch testing until touch is supported on desktop
+			return Target.Platform.IsInGroup(UnrealPlatformGroup.Desktop) && Target.Configuration != UnrealTargetConfiguration.Shipping;
+		}
+	}
 }
diff --git a/Engine/Plugins/Experimental/CommonUI/Source/CommonUI/Private/Input/CommonAnalogCursor.cpp b/Engine/Plugins/Experimental/CommonUI/Source/CommonUI/Private/Input/CommonAnalogCursor.cpp
--- a/Engine/Plugins/Experimental/CommonUI/Source/CommonUI/Private/Input/CommonAnalogCursor.cpp
+++ b/Engine/Plugins/Experimental/CommonUI/Source/CommonUI/Private/Input/CommonAnalogCursor.cpp
@@ -333,9 +333,8 @@ bool FCommonAnalogCursor::HandleMouseButtonDownEvent(FSlateApplication& SlateApp
 {
 	if (FAnalogCursor::IsRelevantInput(PointerEvent))
 	{
-#if PLATFORM_SWITCH	
-		//@todo DanH: Is it intentional or a bug that Switch doesn't swap to Touch input type when the user touches?
-		// Switch doesn't register as switching its input type, so detect touch input here to hide the cursor.
+#if UE_COMMONUI_PLATFORM_SUPPORTS_TOUCH	
+		// Some platforms don't register as switching its input type, so detect touch input here to hide the cursor.
 		if (PointerEvent.IsTouchEvent())
 		{
 			//ClearCenterWidget();
