Fixed UI code incorrectly using IsVisible rather than GetVisible

#rnx

diff --git a/Engine/Plugins/Experimental/CommonUI/Source/CommonUI/Private/CommonButtonBase.cpp b/Engine/Plugins/Experimental/CommonUI/Source/CommonUI/Private/CommonButtonBase.cpp
--- a/Engine/Plugins/Experimental/CommonUI/Source/CommonUI/Private/CommonButtonBase.cpp
+++ b/Engine/Plugins/Experimental/CommonUI/Source/CommonUI/Private/CommonButtonBase.cpp
@@ -1025,7 +1025,7 @@ FReply UCommonButtonBase::HandleButtonDoubleClicked()
 
 void UCommonButtonBase::HandleFocusReceived()
 {
-	if (bShouldSelectUponReceivingFocus && !IsSelected())
+	if (bShouldSelectUponReceivingFocus && !GetSelected())
 	{
 		SetIsSelected(true, false);
 	}
