Common Button didn't notify blueprint about triggering input action changes or triggered input action changes. This should now keep blueprint up to date.
#rnx
#tests: PIE: Ran through some test cases

diff --git a/Engine/Plugins/Experimental/CommonUI/Source/CommonUI/Private/CommonButtonBase.cpp b/Engine/Plugins/Experimental/CommonUI/Source/CommonUI/Private/CommonButtonBase.cpp
--- a/Engine/Plugins/Experimental/CommonUI/Source/CommonUI/Private/CommonButtonBase.cpp
+++ b/Engine/Plugins/Experimental/CommonUI/Source/CommonUI/Private/CommonButtonBase.cpp
@@ -914,6 +914,8 @@ void UCommonButtonBase::SetTriggeredInputAction(const FDataTableRowHandle &Input
 	{
 		TriggeredInputAction = InputActionRow;
 		UpdateInputActionWidget();
+
+		OnTriggeredInputActionChanged(InputActionRow);
 	}
 
 }
@@ -933,6 +935,8 @@ void UCommonButtonBase::SetTriggeringInputAction(const FDataTableRowHandle & Inp
 
 		// Update the Input action widget whenever the triggering input action changes
 		UpdateInputActionWidget();
+
+		OnTriggeringInputActionChanged(InputActionRow);
 	}
 }
 
diff --git a/Engine/Plugins/Experimental/CommonUI/Source/CommonUI/Public/CommonButtonBase.h b/Engine/Plugins/Experimental/CommonUI/Source/CommonUI/Public/CommonButtonBase.h
--- a/Engine/Plugins/Experimental/CommonUI/Source/CommonUI/Public/CommonButtonBase.h
+++ b/Engine/Plugins/Experimental/CommonUI/Source/CommonUI/Public/CommonButtonBase.h
@@ -541,6 +541,9 @@ protected:
 	UFUNCTION(BlueprintImplementableEvent, Category = "Common Button")
 	void OnTriggeredInputActionChanged(const FDataTableRowHandle& NewTriggeredAction);
 
+	UFUNCTION(BlueprintImplementableEvent, Category = "Common Button")
+	void OnTriggeringInputActionChanged(const FDataTableRowHandle& NewTriggeredAction);
+
 	UFUNCTION(BlueprintImplementableEvent, Category = "Common Button")
 	void OnActionProgress(float HeldPercent);
 	
