CommonGame viewport client deprecations to handle the new FInputDeviceID

diff --git a/Engine/Plugins/Experimental/CommonUI/Source/CommonUI/CommonUI.Build.cs b/Engine/Plugins/Experimental/CommonUI/Source/CommonUI/CommonUI.Build.cs
--- a/Engine/Plugins/Experimental/CommonUI/Source/CommonUI/CommonUI.Build.cs
+++ b/Engine/Plugins/Experimental/CommonUI/Source/CommonUI/CommonUI.Build.cs
@@ -11,6 +11,7 @@ public class CommonUI : ModuleRules
 			{
 				"Core",
 				"CoreUObject",
+				"ApplicationCore",
                 "Engine",
 				"InputCore",
 				"Slate",
diff --git a/Engine/Plugins/Experimental/CommonUI/Source/CommonUI/Private/CommonGameViewportClient.cpp b/Engine/Plugins/Experimental/CommonUI/Source/CommonUI/Private/CommonGameViewportClient.cpp
--- a/Engine/Plugins/Experimental/CommonUI/Source/CommonUI/Private/CommonGameViewportClient.cpp
+++ b/Engine/Plugins/Experimental/CommonUI/Source/CommonUI/Private/CommonGameViewportClient.cpp
@@ -6,6 +6,7 @@
 #include "Engine/LocalPlayer.h"
 #include "GameFramework/GameUserSettings.h"
 #include "GameFramework/PlayerController.h"
+#include "GenericPlatform/GenericPlatformInputDeviceMapper.h"
 
 #if WITH_EDITOR
 #include "Editor.h"
@@ -43,9 +44,9 @@ bool UCommonGameViewportClient::InputKey(const FInputKeyEventArgs& InEventArgs)
 #endif
 	{		
 		FReply Result = FReply::Unhandled();
-		if (!OnRerouteInput().ExecuteIfBound(EventArgs.ControllerId, EventArgs.Key, EventArgs.Event, Result))
+		if (!OnRerouteInput().ExecuteIfBound(EventArgs.InputDevice, EventArgs.Key, EventArgs.Event, Result))
 		{
-			HandleRerouteInput(EventArgs.ControllerId, EventArgs.Key, EventArgs.Event, Result);
+			HandleRerouteInput(EventArgs.InputDevice, EventArgs.Key, EventArgs.Event, Result);
 		}
 
 		if (Result.IsEventHandled())
@@ -57,21 +58,20 @@ bool UCommonGameViewportClient::InputKey(const FInputKeyEventArgs& InEventArgs)
 	return Super::InputKey(EventArgs);
 }
 
-bool UCommonGameViewportClient::InputAxis(FViewport* InViewport, int32 UserId, FKey Key, float Delta, float DeltaTime, int32 NumSamples, bool bGamepad)
+bool UCommonGameViewportClient::InputAxis(FViewport* InViewport, FInputDeviceId InputDevice, FKey Key, float Delta, float DeltaTime, int32 NumSamples, bool bGamepad)
 {
-	int32 ControllerId = UserId;
 	FReply RerouteResult = FReply::Unhandled();
 
-	if (!OnRerouteAxis().ExecuteIfBound(ControllerId, Key, Delta, RerouteResult))
+	if (!OnRerouteAxis().ExecuteIfBound(InputDevice, Key, Delta, RerouteResult))
 	{
-		HandleRerouteAxis(ControllerId, Key, Delta, RerouteResult);
+		HandleRerouteAxis(InputDevice, Key, Delta, RerouteResult);
 	}
 
 	if (RerouteResult.IsEventHandled())
 	{
 		return true;
 	}
-	return Super::InputAxis(InViewport, ControllerId, Key, Delta, DeltaTime, NumSamples, bGamepad);
+	return Super::InputAxis(InViewport, InputDevice, Key, Delta, DeltaTime, NumSamples, bGamepad);
 }
 
 bool UCommonGameViewportClient::InputTouch(FViewport* InViewport, int32 ControllerId, uint32 Handle, ETouchType::Type Type, const FVector2D& TouchLocation, float Force, FDateTime DeviceTimestamp, uint32 TouchpadIndex)
@@ -95,9 +95,10 @@ bool UCommonGameViewportClient::InputTouch(FViewport* InViewport, int32 Controll
 	return Super::InputTouch(InViewport, ControllerId, Handle, Type, TouchLocation, Force, DeviceTimestamp, TouchpadIndex);
 }
 
-void UCommonGameViewportClient::HandleRerouteInput(int32 ControllerId, FKey Key, EInputEvent EventType, FReply& Reply)
+void UCommonGameViewportClient::HandleRerouteInput(FInputDeviceId DeviceId, FKey Key, EInputEvent EventType, FReply& Reply)
 {
-	ULocalPlayer* LocalPlayer = GameInstance->FindLocalPlayerFromControllerId(ControllerId);
+	FPlatformUserId OwningPlatformUser = IPlatformInputDeviceMapper::Get().GetUserForInputDevice(DeviceId);
+	ULocalPlayer* LocalPlayer = GameInstance->FindLocalPlayerFromPlatformUserId(OwningPlatformUser);
 	Reply = FReply::Unhandled();
 
 	if (LocalPlayer)
@@ -111,19 +112,31 @@ void UCommonGameViewportClient::HandleRerouteInput(int32 ControllerId, FKey Key,
 				// We need to set the reply as handled otherwise the input won't actually be blocked from reaching the viewport.
 				Reply = FReply::Handled();
 				// Notify interested parties that we blocked the input.
-				OnRerouteBlockedInput().ExecuteIfBound(ControllerId, Key, EventType, Reply);
+				OnRerouteBlockedInput().ExecuteIfBound(DeviceId, Key, EventType, Reply);
 			}
 			else if (InputResult == ERouteUIInputResult::Handled)
 			{
 				Reply = FReply::Handled();
 			}
 		}
-	}
+	}	
 }
 
-void UCommonGameViewportClient::HandleRerouteAxis(int32 ControllerId, FKey Key, float Delta, FReply& Reply)
+void UCommonGameViewportClient::HandleRerouteInput(int32 ControllerId, FKey Key, EInputEvent EventType, FReply& Reply)
 {
-	ULocalPlayer* LocalPlayer = GameInstance->FindLocalPlayerFromControllerId(ControllerId);
+	// Remap the old int32 ControllerId to the new platform user and input device ID
+	FPlatformUserId UserId = FGenericPlatformMisc::GetPlatformUserForUserIndex(ControllerId);
+	FInputDeviceId DeviceID = INPUTDEVICEID_NONE;
+	IPlatformInputDeviceMapper::Get().RemapControllerIdToPlatformUserAndDevice(ControllerId, UserId, DeviceID);
+	return HandleRerouteInput(DeviceID, Key, EventType, Reply);
+}
+
+void UCommonGameViewportClient::HandleRerouteAxis(FInputDeviceId DeviceId, FKey Key, float Delta, FReply& Reply)
+{
+	// Get the ownign platform user for this input device and their local player
+	FPlatformUserId OwningPlatformUser = IPlatformInputDeviceMapper::Get().GetUserForInputDevice(DeviceId);
+	ULocalPlayer* LocalPlayer = GameInstance->FindLocalPlayerFromPlatformUserId(OwningPlatformUser);
+	
 	Reply = FReply::Unhandled();
 
 	if (LocalPlayer)
@@ -140,6 +153,16 @@ void UCommonGameViewportClient::HandleRerouteAxis(int32 ControllerId, FKey Key,
 	}
 }
 
+void UCommonGameViewportClient::HandleRerouteAxis(int32 ControllerId, FKey Key, float Delta, FReply& Reply)
+{
+	// Remap the old int32 ControllerId to the new platform user and input device ID
+	FPlatformUserId UserId = FGenericPlatformMisc::GetPlatformUserForUserIndex(ControllerId);
+	FInputDeviceId DeviceID = INPUTDEVICEID_NONE;
+	IPlatformInputDeviceMapper::Get().RemapControllerIdToPlatformUserAndDevice(ControllerId, UserId, DeviceID);
+	
+	return HandleRerouteAxis(DeviceID, Key, Delta, Reply);
+}
+
 void UCommonGameViewportClient::HandleRerouteTouch(int32 ControllerId, uint32 TouchId, ETouchType::Type TouchType, const FVector2D& TouchLocation, FReply& Reply)
 {
 	ULocalPlayer* LocalPlayer = GameInstance->FindLocalPlayerFromControllerId(ControllerId);
diff --git a/Engine/Plugins/Experimental/CommonUI/Source/CommonUI/Public/CommonGameViewportClient.h b/Engine/Plugins/Experimental/CommonUI/Source/CommonUI/Public/CommonGameViewportClient.h
--- a/Engine/Plugins/Experimental/CommonUI/Source/CommonUI/Public/CommonGameViewportClient.h
+++ b/Engine/Plugins/Experimental/CommonUI/Source/CommonUI/Public/CommonGameViewportClient.h
@@ -8,8 +8,8 @@
 #include "Engine/GameViewportClient.h"
 #include "CommonGameViewportClient.generated.h"
 
-DECLARE_DELEGATE_FourParams(FOnRerouteInputDelegate, int32 /* ControllerId */, FKey /* Key */, EInputEvent /* EventType */, FReply& /* Reply */);
-DECLARE_DELEGATE_FourParams(FOnRerouteAxisDelegate, int32 /* ControllerId */, FKey /* Key */, float /* Delta */, FReply& /* Reply */);
+DECLARE_DELEGATE_FourParams(FOnRerouteInputDelegate, FInputDeviceId /* InputDeviceId */, FKey /* Key */, EInputEvent /* EventType */, FReply& /* Reply */);
+DECLARE_DELEGATE_FourParams(FOnRerouteAxisDelegate, FInputDeviceId /* InputDeviceId */, FKey /* Key */, float /* Delta */, FReply& /* Reply */);
 DECLARE_DELEGATE_FiveParams(FOnRerouteTouchDelegate, int32 /* ControllerId */, uint32 /* TouchId */, ETouchType::Type /* TouchType */, const FVector2D& /* TouchLocation */, FReply& /* Reply */);
 
 /**
@@ -27,7 +27,7 @@ public:
 
 	// UGameViewportClient interface begin
 	virtual bool InputKey(const FInputKeyEventArgs& EventArgs) override;
-	virtual bool InputAxis(FViewport* InViewport, int32 UserId, FKey Key, float Delta, float DeltaTime, int32 NumSamples = 1, bool bGamepad = false) override;
+	virtual bool InputAxis(FViewport* InViewport, FInputDeviceId InputDevice, FKey Key, float Delta, float DeltaTime, int32 NumSamples, bool bGamepad) override;
 	virtual bool InputTouch(FViewport* InViewport, int32 ControllerId, uint32 Handle, ETouchType::Type Type, const FVector2D& TouchLocation, float Force, FDateTime DeviceTimestamp, uint32 TouchpadIndex) override;
 	// UGameViewportClient interface end
 
@@ -38,11 +38,19 @@ public:
 	FOnRerouteInputDelegate& OnRerouteBlockedInput() { return RerouteBlockedInput; }
 
 	/** Default Handler for Key input. */
+	UE_DEPRECATED(5.1, "This version of HandleRerouteInput has been deprecated. Please use the version that takes an FInputDeviceId instead")
 	virtual void HandleRerouteInput(int32 ControllerId, FKey Key, EInputEvent EventType, FReply& Reply);
 
+	/** Default Handler for Key input. */
+	virtual void HandleRerouteInput(FInputDeviceId DeviceId, FKey Key, EInputEvent EventType, FReply& Reply);
+
 	/** Default Handler for Axis input. */
+	UE_DEPRECATED(5.1, "This version of HandleRerouteAxis has been deprecated. Please use the version that takes an FInputDeviceId instead")
 	virtual void HandleRerouteAxis(int32 ControllerId, FKey Key, float Delta, FReply& Reply);
 
+	/** Default Handler for Axis input. */
+	virtual void HandleRerouteAxis(FInputDeviceId DeviceId, FKey Key, float Delta, FReply& Reply);
+
 	/** Default Handler for Touch input. */
 	virtual void HandleRerouteTouch(int32 ControllerId, uint32 TouchId, ETouchType::Type TouchType, const FVector2D& TouchLocation, FReply& Reply);
 
