Allow CommonButtonGroupBase to optionally suppress button audio
#tests: Verified button audio was silenced when calling SelectButtonAtIndex with bAllowSound = false

diff --git a/Engine/Plugins/Experimental/CommonUI/Source/CommonUI/Private/Groups/CommonButtonGroupBase.cpp b/Engine/Plugins/Experimental/CommonUI/Source/CommonUI/Private/Groups/CommonButtonGroupBase.cpp
--- a/Engine/Plugins/Experimental/CommonUI/Source/CommonUI/Private/Groups/CommonButtonGroupBase.cpp
+++ b/Engine/Plugins/Experimental/CommonUI/Source/CommonUI/Private/Groups/CommonButtonGroupBase.cpp
@@ -140,7 +140,7 @@ void UCommonButtonGroupBase::SelectPreviousButton(bool bAllowWrap /*= true*/)
 	SelectPrevButtonRecursive(SelectedButtonIndex, bAllowWrap);
 }
 
-void UCommonButtonGroupBase::SelectButtonAtIndex(int32 ButtonIndex)
+void UCommonButtonGroupBase::SelectButtonAtIndex(int32 ButtonIndex, const bool bAllowSound)
 {
 	if (ButtonIndex < 0 || ButtonIndex >= Buttons.Num())
 	{
@@ -152,7 +152,7 @@ void UCommonButtonGroupBase::SelectButtonAtIndex(int32 ButtonIndex)
 		if (ButtonToSelect && !ButtonToSelect->GetSelected() && ButtonToSelect->GetIsEnabled())
 		{
 			SelectedButtonIndex = ButtonIndex;
-			ButtonToSelect->SetSelectedInternal(true);
+			ButtonToSelect->SetSelectedInternal(true, bAllowSound);
 		}
 	}
 }
diff --git a/Engine/Plugins/Experimental/CommonUI/Source/CommonUI/Public/Groups/CommonButtonGroupBase.h b/Engine/Plugins/Experimental/CommonUI/Source/CommonUI/Public/Groups/CommonButtonGroupBase.h
--- a/Engine/Plugins/Experimental/CommonUI/Source/CommonUI/Public/Groups/CommonButtonGroupBase.h
+++ b/Engine/Plugins/Experimental/CommonUI/Source/CommonUI/Public/Groups/CommonButtonGroupBase.h
@@ -54,9 +54,10 @@ public:
 	/**
 	 * Selects a button at a specific index in the group. Clears all selection if given an invalid index.
 	 * @param ButtonIndex The index of the button in the group to select
+	 * @param bAllowSound Whether the selected button should play its click sound
 	 */
 	UFUNCTION(BlueprintCallable, Category = BaseButtonGroup)
-	void SelectButtonAtIndex(int32 ButtonIndex);
+	void SelectButtonAtIndex(int32 ButtonIndex, const bool bAllowSound = true);
 
 	/**
 	 * Get the index of the currently selected button, if any.
