Updated common button base to click before selection as the default but we give a new option to reverse the order.
#rnx
#tests: ran a few test cases

diff --git a/Engine/Plugins/Experimental/CommonUI/Source/CommonUI/Private/CommonButtonBase.cpp b/Engine/Plugins/Experimental/CommonUI/Source/CommonUI/Private/CommonButtonBase.cpp
--- a/Engine/Plugins/Experimental/CommonUI/Source/CommonUI/Private/CommonButtonBase.cpp
+++ b/Engine/Plugins/Experimental/CommonUI/Source/CommonUI/Private/CommonButtonBase.cpp
@@ -312,6 +312,7 @@ UCommonButtonBase::UCommonButtonBase(const FObjectInitializer& ObjectInitializer
 	, bSelectable(false)
 	, bShouldSelectUponReceivingFocus(false)
 	, bToggleable(false)
+	, bTriggerClickedAfterSelection(false)
 	, bDisplayInputActionWhenNotInteractable(true)
 	, bShouldUseFallbackDefaultInputAction(true)
 	, bSelected(false)
@@ -1012,9 +1013,16 @@ void UCommonButtonBase::HandleButtonClicked()
 	// Guard against this case.
 	if (IsInteractionEnabled())
 	{
-		SetIsSelected(!bSelected, false);
-
-		NativeOnClicked();
+		if (bTriggerClickedAfterSelection)
+		{
+			SetIsSelected(!bSelected, false);
+			NativeOnClicked();
+		}
+		else
+		{
+			NativeOnClicked();
+			SetIsSelected(!bSelected, false);
+		}
 
 		ExecuteTriggeredInput();
 	}
diff --git a/Engine/Plugins/Experimental/CommonUI/Source/CommonUI/Public/CommonButtonBase.h b/Engine/Plugins/Experimental/CommonUI/Source/CommonUI/Public/CommonButtonBase.h
--- a/Engine/Plugins/Experimental/CommonUI/Source/CommonUI/Public/CommonButtonBase.h
+++ b/Engine/Plugins/Experimental/CommonUI/Source/CommonUI/Public/CommonButtonBase.h
@@ -606,6 +606,9 @@ protected:
 	UPROPERTY(EditAnywhere, BlueprintReadOnly, Category = Selection, meta = (ExposeOnSpawn = true, EditCondition = "bSelectable"))
 	uint8 bToggleable:1;
 
+	UPROPERTY(EditAnywhere, BlueprintReadOnly, Category = Selection, meta = (ExposeOnSpawn = true, EditCondition = "bSelectable"))
+	uint8 bTriggerClickedAfterSelection:1;
+
 	/** True if the input action should be displayed when the button is not interactable */
 	UPROPERTY(EditAnywhere, BlueprintReadOnly, Category = Input, meta = (ExposeOnSpawn = true))
 	uint8 bDisplayInputActionWhenNotInteractable:1;
