Added bAutoCollapseWithEmptyText to UCommonRichTextBlock to mirror UCommonTextBlock.

diff --git a/Engine/Plugins/Experimental/CommonUI/Source/CommonUI/Private/CommonRichTextBlock.cpp b/Engine/Plugins/Experimental/CommonUI/Source/CommonUI/Private/CommonRichTextBlock.cpp
--- a/Engine/Plugins/Experimental/CommonUI/Source/CommonUI/Private/CommonRichTextBlock.cpp
+++ b/Engine/Plugins/Experimental/CommonUI/Source/CommonUI/Private/CommonRichTextBlock.cpp
@@ -500,6 +500,11 @@ void UCommonRichTextBlock::SetText(const FText& InText)
 		}
 		MyRichTextBlock->SetText(InText);
 	}
+
+	if (bAutoCollapseWithEmptyText)
+	{
+		SetVisibility(InText.IsEmpty() ? ESlateVisibility::Collapsed : ESlateVisibility::SelfHitTestInvisible);
+	}
 }
 
 void UCommonRichTextBlock::SetScrollingEnabled(bool bInIsScrollingEnabled)
@@ -523,15 +528,22 @@ const FText UCommonRichTextBlock::GetPaletteCategory()
 
 bool UCommonRichTextBlock::CanEditChange(const FProperty* InProperty) const
 {
+	const FName PropertyName = InProperty->GetFName();
+
 	if (DefaultTextStyleOverrideClass)
 	{
-		const FName PropertyName = InProperty->GetFName();
 		if (PropertyName.IsEqual(GET_MEMBER_NAME_CHECKED(UCommonRichTextBlock, DefaultTextStyleOverride)) ||
 			PropertyName.IsEqual(GET_MEMBER_NAME_CHECKED(UCommonRichTextBlock, MinDesiredWidth)))
 		{
 			return false;
 		}
 	}
+
+	if (bAutoCollapseWithEmptyText && PropertyName.IsEqual(GET_MEMBER_NAME_CHECKED(UCommonRichTextBlock, Visibility)))
+	{
+		return false;
+	}
+
 	return Super::CanEditChange(InProperty);
 }
 
@@ -590,6 +602,11 @@ TSharedRef<SWidget> UCommonRichTextBlock::RebuildWidget()
 void UCommonRichTextBlock::SynchronizeProperties()
 {
 	Super::SynchronizeProperties();
+
+	if (bAutoCollapseWithEmptyText)
+	{
+		SetVisibility(GetText().IsEmpty() ? ESlateVisibility::Collapsed : ESlateVisibility::SelfHitTestInvisible);
+	}
 	
 	RefreshOverrideStyle();
 
diff --git a/Engine/Plugins/Experimental/CommonUI/Source/CommonUI/Public/CommonRichTextBlock.h b/Engine/Plugins/Experimental/CommonUI/Source/CommonUI/Public/CommonRichTextBlock.h
--- a/Engine/Plugins/Experimental/CommonUI/Source/CommonUI/Public/CommonRichTextBlock.h
+++ b/Engine/Plugins/Experimental/CommonUI/Source/CommonUI/Public/CommonRichTextBlock.h
@@ -89,5 +89,9 @@ private:
 	UPROPERTY(meta = (DeprecatedProperty, DeprecationMessage = "bDisplayAllCaps is deprecated. Please use TextTransformPolicy instead."))
 	bool bDisplayAllCaps_DEPRECATED = false;
 
+	/** True to automatically collapse this rich text block when set to display an empty string. Conversely, will be SelfHitTestInvisible when showing a non-empty string. */
+	UPROPERTY(EditAnywhere, BlueprintReadOnly, Category = Appearance, meta = (AllowPrivateAccess = true))
+	bool bAutoCollapseWithEmptyText = false;
+
 	TSharedPtr<STextScroller> MyTextScroller;
 };
