Fix CIS error.

diff --git a/Engine/Plugins/Experimental/CommonUI/Source/CommonUI/Private/CommonRichTextBlock.cpp b/Engine/Plugins/Experimental/CommonUI/Source/CommonUI/Private/CommonRichTextBlock.cpp
--- a/Engine/Plugins/Experimental/CommonUI/Source/CommonUI/Private/CommonRichTextBlock.cpp
+++ b/Engine/Plugins/Experimental/CommonUI/Source/CommonUI/Private/CommonRichTextBlock.cpp
@@ -539,10 +539,12 @@ bool UCommonRichTextBlock::CanEditChange(const FProperty* InProperty) const
 		}
 	}
 
+PRAGMA_DISABLE_DEPRECATION_WARNINGS
 	if (bAutoCollapseWithEmptyText && PropertyName.IsEqual(GET_MEMBER_NAME_CHECKED(UCommonRichTextBlock, Visibility)))
 	{
 		return false;
 	}
+PRAGMA_ENABLE_DEPRECATION_WARNINGS
 
 	return Super::CanEditChange(InProperty);
 }
