Mobile interaction is disfunctional in Main - Reverted a change in the code that caused mobile interactions to malfunction.

diff --git a/Engine/Plugins/Experimental/CommonUI/Source/CommonInput/Private/CommonInputSubsystem.cpp b/Engine/Plugins/Experimental/CommonUI/Source/CommonInput/Private/CommonInputSubsystem.cpp
--- a/Engine/Plugins/Experimental/CommonUI/Source/CommonInput/Private/CommonInputSubsystem.cpp
+++ b/Engine/Plugins/Experimental/CommonUI/Source/CommonInput/Private/CommonInputSubsystem.cpp
@@ -727,6 +727,9 @@ bool UCommonInputSubsystem::PlatformSupportsInputType(ECommonInputType InInputTy
 		case ECommonInputType::Touch:
 		{
 			bPlatformSupportsInput &= UE_COMMONINPUT_PLATFORM_SUPPORTS_TOUCH != 0;
+#if PLATFORM_DESKTOP && !UE_BUILD_SHIPPING
+			bPlatformSupportsInput = true; // Enable touch in debug editor
+#endif
 		}
 		break;
 		case ECommonInputType::Gamepad:
