Style Sheet:
- Prototype version of the Common UI Style Sheet

diff --git a/Engine/Plugins/Experimental/CommonUI/Source/CommonUI/Private/CommonTextBlock.cpp b/Engine/Plugins/Experimental/CommonUI/Source/CommonUI/Private/CommonTextBlock.cpp
--- a/Engine/Plugins/Experimental/CommonUI/Source/CommonUI/Private/CommonTextBlock.cpp
+++ b/Engine/Plugins/Experimental/CommonUI/Source/CommonUI/Private/CommonTextBlock.cpp
@@ -9,6 +9,7 @@
 #include "Widgets/Text/STextBlock.h"
 #include "Layout/LayoutUtils.h"
 #include "Types/ReflectionMetadata.h"
+#include "StyleSheet/CommonStyleSheet.h"
 
 void STextScroller::Construct(const FArguments& InArgs)
 {
@@ -521,12 +522,27 @@ void UCommonTextBlock::SetTextCase(bool bUseAllCaps)
 	SynchronizeProperties();
 }
 
+void UCommonTextBlock::SetLineHeightPercentage(float InLineHeightPercentage)
+{
+	LineHeightPercentage = InLineHeightPercentage;
+}
+
 void UCommonTextBlock::SetStyle(TSubclassOf<UCommonTextStyle> InStyle)
 {
 	Style = InStyle;
 	SynchronizeProperties();
 }
 
+const FMargin& UCommonTextBlock::GetMargin()
+{
+	return Margin;
+}
+
+void UCommonTextBlock::SetMargin(const FMargin& InMargin)
+{
+	Margin = InMargin;
+}
+
 void UCommonTextBlock::ResetScrollState()
 {
 	if (TextScroller.IsValid())
@@ -612,6 +628,11 @@ void UCommonTextBlock::UpdateFromStyle()
 			SetShadowColorAndOpacity(FLinearColor::Transparent);
 		}
 	}
+
+	if (UCommonStyleSheet* StyleSheetPtr = StyleSheet.Get())
+	{
+		StyleSheetPtr->Apply(this);
+	}
 }
 
 void UCommonTextBlock::ReleaseSlateResources(bool bReleaseChildren)
diff --git a/Engine/Plugins/Experimental/CommonUI/Source/CommonUI/Private/StyleSheet/CommonStyleSheet.cpp b/Engine/Plugins/Experimental/CommonUI/Source/CommonUI/Private/StyleSheet/CommonStyleSheet.cpp
new file mode 100644
--- /dev/null
+++ b/Engine/Plugins/Experimental/CommonUI/Source/CommonUI/Private/StyleSheet/CommonStyleSheet.cpp
@@ -0,0 +1,125 @@
+// Copyright Epic Games, Inc. All Rights Reserved.
+
+#include "StyleSheet/CommonStyleSheet.h"
+#include "StyleSheet/CommonStyleSheetTypes.h"
+
+#include "Blueprint/UserWidget.h"
+#include "CommonTextBlock.h"
+#include "Components/Widget.h"
+
+void UCommonStyleSheet::Apply(UWidget* Widget)
+{
+	// FIXME: Current usage of the SetX functions can results in Apply being called
+	// resulting in an infinite loop. To prevent this, we use this flag.
+	if (bIsApplying)
+	{
+		return;
+	}
+
+	bIsApplying = true;
+	if (UCommonTextBlock* TextBlock = Cast<UCommonTextBlock>(Widget))
+	{
+		TryApplyColorAndOpacity(TextBlock);
+		TryApplyLineHeightPercentage(TextBlock);
+		TryApplyFont(TextBlock);
+		TryApplyMargin(TextBlock);
+	}
+	bIsApplying = false;
+}
+
+void UCommonStyleSheet::TryApplyColorAndOpacity(UCommonTextBlock* TextBlock) const
+{
+	const UCommonStyleSheetTypeColor* TypeColor = FindStyleSheetProperty<UCommonStyleSheetTypeColor>(this);
+	FLinearColor Color = TypeColor ? TypeColor->Color : TextBlock->GetColorAndOpacity().GetSpecifiedColor();
+
+	if (const UCommonStyleSheetTypeOpacity* TypeOpacity = FindStyleSheetProperty<UCommonStyleSheetTypeOpacity>(this))
+	{
+		Color.A = TypeOpacity->Opacity;
+	}
+
+	TextBlock->SetColorAndOpacity(Color);
+}
+
+void UCommonStyleSheet::TryApplyLineHeightPercentage(UCommonTextBlock* TextBlock) const
+{
+	if (const UCommonStyleSheetTypeLineHeightPercentage* TypeLineHeightPercentage = FindStyleSheetProperty<UCommonStyleSheetTypeLineHeightPercentage>(this))
+	{
+		TextBlock->SetLineHeightPercentage(TypeLineHeightPercentage->LineHeightPercentage);
+	}
+}
+
+void UCommonStyleSheet::TryApplyFont(UCommonTextBlock* TextBlock) const
+{
+	FSlateFontInfo FontInfo = TextBlock->GetFont();
+
+	if (const UCommonStyleSheetTypeFontTypeface* TypeFontTypeface = FindStyleSheetProperty<UCommonStyleSheetTypeFontTypeface>(this))
+	{
+		FontInfo.FontObject = TypeFontTypeface->Typeface.FontObject;
+		FontInfo.TypefaceFontName = TypeFontTypeface->Typeface.TypefaceFontName;
+	}
+
+	if (const UCommonStyleSheetTypeFontSize* TypeFontSize = FindStyleSheetProperty<UCommonStyleSheetTypeFontSize>(this))
+	{
+		FontInfo.Size = TypeFontSize->Size;
+	}
+
+	if (const UCommonStyleSheetTypeFontLetterSpacing* TypeFontLetterSpacing = FindStyleSheetProperty<UCommonStyleSheetTypeFontLetterSpacing>(this))
+	{
+		FontInfo.LetterSpacing = TypeFontLetterSpacing->LetterSpacing;
+	}
+
+	TextBlock->SetFont(FontInfo);
+}
+
+void UCommonStyleSheet::TryApplyMargin(UCommonTextBlock* TextBlock) const
+{
+	FMargin Margin = TextBlock->GetMargin();
+
+	if (const UCommonStyleSheetTypeMarginLeft* TypeMarginLeft = FindStyleSheetProperty<UCommonStyleSheetTypeMarginLeft>(this))
+	{
+		Margin.Left = TypeMarginLeft->Left;
+	}
+
+	if (const UCommonStyleSheetTypeMarginRight* TypeMarginRight = FindStyleSheetProperty<UCommonStyleSheetTypeMarginRight>(this))
+	{
+		Margin.Right = TypeMarginRight->Right;
+	}
+
+	if (const UCommonStyleSheetTypeMarginTop* TypeMarginTop = FindStyleSheetProperty<UCommonStyleSheetTypeMarginTop>(this))
+	{
+		Margin.Top = TypeMarginTop->Top;
+	}
+
+	if (const UCommonStyleSheetTypeMarginBottom* TypeMarginBottom = FindStyleSheetProperty<UCommonStyleSheetTypeMarginBottom>(this))
+	{
+		Margin.Bottom = TypeMarginBottom->Bottom;
+	}
+
+	TextBlock->SetMargin(Margin);
+}
+
+template <typename T>
+const T* UCommonStyleSheet::FindStyleSheetProperty(const UCommonStyleSheet* StyleSheet) const
+{
+	for (const TObjectPtr<UCommonStyleSheetTypeBase>& Property : StyleSheet->Properties)
+	{
+		const UCommonStyleSheetTypeBase* PropertyPtr = Property.Get();
+		if (PropertyPtr && PropertyPtr->bIsEnabled && PropertyPtr->IsA(T::StaticClass()))
+		{
+			return Cast<T>(PropertyPtr);
+		}
+	}
+
+	for (const UCommonStyleSheet* ImportedStyleSheet : StyleSheet->ImportedStyleSheets)
+	{
+		if (ImportedStyleSheet)
+		{
+			if (const T* Property = FindStyleSheetProperty<T>(ImportedStyleSheet))
+			{
+				return Property;
+			}
+		}
+	}
+
+	return nullptr;
+}
diff --git a/Engine/Plugins/Experimental/CommonUI/Source/CommonUI/Public/CommonTextBlock.h b/Engine/Plugins/Experimental/CommonUI/Source/CommonUI/Public/CommonTextBlock.h
--- a/Engine/Plugins/Experimental/CommonUI/Source/CommonUI/Public/CommonTextBlock.h
+++ b/Engine/Plugins/Experimental/CommonUI/Source/CommonUI/Public/CommonTextBlock.h
@@ -7,6 +7,8 @@
 #include "Components/TextBlock.h"
 #include "CommonTextBlock.generated.h"
 
+class UCommonStyleSheet;
+
 /* 
  * ---- All properties must be EditDefaultsOnly, BlueprintReadOnly !!! -----
  * We return the CDO to blueprints, so we cannot allow any changes (blueprint doesn't support const variables)
@@ -180,9 +182,18 @@ public:
 	UFUNCTION(BlueprintCallable, Category = "Common Text")
 	void SetTextCase(bool bUseAllCaps);
 
+	UFUNCTION(BlueprintCallable, Category = "Common Text")
+	void SetLineHeightPercentage(float InLineHeightPercentage);
+
 	UFUNCTION(BlueprintCallable, Category = "Common Text")
 	void SetStyle(TSubclassOf<UCommonTextStyle> InStyle);
 
+	UFUNCTION(BlueprintCallable, Category = "Common Text")
+	const FMargin& GetMargin();
+
+	UFUNCTION(BlueprintCallable, Category = "Common Text")
+	void SetMargin(const FMargin& InMargin);
+
 	UFUNCTION(BlueprintCallable, Category = "Common Text|Scroll Style")
 	void ResetScrollState();
 
@@ -208,6 +219,10 @@ private:
 	UPROPERTY(EditAnywhere, BlueprintReadOnly, Category = CommonText, meta = (ExposeOnSpawn = true, AllowPrivateAccess = true))
 	TSubclassOf<UCommonTextScrollStyle> ScrollStyle;
 
+	/** References the scroll style asset to use, no reference disables scrolling*/
+	UPROPERTY(EditAnywhere, BlueprintReadOnly, Category = CommonText, meta = (ExposeOnSpawn = true, AllowPrivateAccess = true))
+	TObjectPtr<UCommonStyleSheet> StyleSheet;
+
 	/** If scrolling is enabled/disabled initially, this can be updated in blueprint */
 	UPROPERTY(EditAnywhere, BlueprintReadOnly, Category = CommonText, meta = (ExposeOnSpawn = true, AllowPrivateAccess = true))
 	bool bIsScrollingEnabled = true;
diff --git a/Engine/Plugins/Experimental/CommonUI/Source/CommonUI/Public/StyleSheet/CommonStyleSheet.h b/Engine/Plugins/Experimental/CommonUI/Source/CommonUI/Public/StyleSheet/CommonStyleSheet.h
new file mode 100644
--- /dev/null
+++ b/Engine/Plugins/Experimental/CommonUI/Source/CommonUI/Public/StyleSheet/CommonStyleSheet.h
@@ -0,0 +1,38 @@
+// Copyright Epic Games, Inc. All Rights Reserved.
+
+#pragma once
+
+#include "Engine/DataAsset.h"
+
+#include "CommonStyleSheet.generated.h"
+
+class UCommonStyleSheetTypeBase;
+class UCommonTextBlock;
+class UWidget;
+
+UCLASS(BlueprintType)
+class UCommonStyleSheet : public UDataAsset
+{
+	GENERATED_BODY()
+
+public:
+	void Apply(UWidget* Widget);
+
+private:
+	void TryApplyColorAndOpacity(UCommonTextBlock* TextBlock) const;
+	void TryApplyLineHeightPercentage(UCommonTextBlock* TextBlock) const;
+	void TryApplyFont(UCommonTextBlock* TextBlock) const;
+	void TryApplyMargin(UCommonTextBlock* TextBlock) const;
+
+	template <typename T>
+	const T* FindStyleSheetProperty(const UCommonStyleSheet* StyleSheet) const;
+
+private:
+	UPROPERTY(EditDefaultsOnly, Instanced, Category = "Properties")
+	TArray<TObjectPtr<UCommonStyleSheetTypeBase>> Properties;
+
+	UPROPERTY(EditDefaultsOnly, Category = "Import")
+	TArray<UCommonStyleSheet*> ImportedStyleSheets;
+
+	bool bIsApplying = false;
+};
diff --git a/Engine/Plugins/Experimental/CommonUI/Source/CommonUI/Public/StyleSheet/CommonStyleSheetTypes.h b/Engine/Plugins/Experimental/CommonUI/Source/CommonUI/Public/StyleSheet/CommonStyleSheetTypes.h
new file mode 100644
--- /dev/null
+++ b/Engine/Plugins/Experimental/CommonUI/Source/CommonUI/Public/StyleSheet/CommonStyleSheetTypes.h
@@ -0,0 +1,125 @@
+// Copyright Epic Games, Inc. All Rights Reserved.
+
+#pragma once
+
+#include "Fonts/SlateFontInfo.h"
+#include "Math/Color.h"
+
+#include "CommonStyleSheetTypes.generated.h"
+
+UCLASS(Abstract, EditInlineNew)
+class UCommonStyleSheetTypeBase : public UObject
+{
+	GENERATED_BODY()
+
+public:
+	UPROPERTY(EditDefaultsOnly, Category = "Value", meta = (InlineEditConditionToggle))
+	bool bIsEnabled = true;
+};
+
+UCLASS(meta = (DisplayName = "Color"))
+class UCommonStyleSheetTypeColor : public UCommonStyleSheetTypeBase
+{
+	GENERATED_BODY()
+
+public:
+	UPROPERTY(EditDefaultsOnly, Category = "Value", meta = (HideAlphaChannel, EditCondition = "bIsEnabled"))
+	FLinearColor Color = FLinearColor(1.0f, 1.0f, 1.0f, 1.0f);
+};
+
+UCLASS(meta = (DisplayName = "Opacity"))
+class UCommonStyleSheetTypeOpacity : public UCommonStyleSheetTypeBase
+{
+	GENERATED_BODY()
+
+public:
+	UPROPERTY(EditDefaultsOnly, Category = "Value", meta = (ClampMin = 0.0f, ClampMax = 1.0f, EditCondition = "bIsEnabled"))
+	float Opacity = 1.0f;
+};
+
+UCLASS(meta = (DisplayName = "Line Height Percentage"))
+class UCommonStyleSheetTypeLineHeightPercentage : public UCommonStyleSheetTypeBase
+{
+	GENERATED_BODY()
+
+public:
+	UPROPERTY(EditDefaultsOnly, Category = "Value", meta = (EditCondition = "bIsEnabled"))
+	float LineHeightPercentage = 1.0f;
+};
+
+UCLASS(meta = (DisplayName = "Font - Typeface"))
+class UCommonStyleSheetTypeFontTypeface : public UCommonStyleSheetTypeBase
+{
+	GENERATED_BODY()
+
+public:
+//	FIXME: Should be used instead of SlateFontInfo but no time to do the customization
+// 	UPROPERTY(EditDefaultsOnly, Category = "Value", meta = (AllowedClasses = "/Script/Engine.Font"))
+// 	TObjectPtr<const UObject> FontFamily;
+// 
+// 	UPROPERTY(EditDefaultsOnly, Category = "Value")
+// 	FName Typeface;
+
+	UPROPERTY(EditDefaultsOnly, Category = "Value", meta = (EditCondition = "bIsEnabled"))
+	FSlateFontInfo Typeface;
+};
+
+UCLASS(meta = (DisplayName = "Font - Size"))
+class UCommonStyleSheetTypeFontSize : public UCommonStyleSheetTypeBase
+{
+	GENERATED_BODY()
+
+public:
+	UPROPERTY(EditDefaultsOnly, Category = "Value", meta = (ClampMin = 1, ClampMax = 1000, EditCondition = "bIsEnabled"))
+	int32 Size = 24;
+};
+
+UCLASS(meta = (DisplayName = "Font - Letter Spacing"))
+class UCommonStyleSheetTypeFontLetterSpacing : public UCommonStyleSheetTypeBase
+{
+	GENERATED_BODY()
+
+public:
+	UPROPERTY(EditDefaultsOnly, Category = "Value", meta = (ClampMin = -1000, ClampMax = 10000, EditCondition = "bIsEnabled"))
+	int32 LetterSpacing;
+};
+
+UCLASS(meta = (DisplayName = "Margin - Left"))
+class UCommonStyleSheetTypeMarginLeft : public UCommonStyleSheetTypeBase
+{
+	GENERATED_BODY()
+
+public:
+	UPROPERTY(EditDefaultsOnly, Category = "Value", meta = (EditCondition = "bIsEnabled"))
+	float Left = 0.0f;
+};
+
+UCLASS(meta = (DisplayName = "Margin - Right"))
+class UCommonStyleSheetTypeMarginRight : public UCommonStyleSheetTypeBase
+{
+	GENERATED_BODY()
+
+public:
+	UPROPERTY(EditDefaultsOnly, Category = "Value", meta = (EditCondition = "bIsEnabled"))
+	float Right = 0.0f;
+};
+
+UCLASS(meta = (DisplayName = "Margin - Top"))
+class UCommonStyleSheetTypeMarginTop : public UCommonStyleSheetTypeBase
+{
+	GENERATED_BODY()
+
+public:
+	UPROPERTY(EditDefaultsOnly, Category = "Value", meta = (EditCondition = "bIsEnabled"))
+	float Top = 0.0f;
+};
+
+UCLASS(meta = (DisplayName = "Margin - Bottom"))
+class UCommonStyleSheetTypeMarginBottom : public UCommonStyleSheetTypeBase
+{
+	GENERATED_BODY()
+
+public:
+	UPROPERTY(EditDefaultsOnly, Category = "Value", meta = (EditCondition = "bIsEnabled"))
+	float Bottom = 0.0f;
+};
