Change the Common UI game viewport to check the override delegate before letting the UI block input, this is needed for things like press start screens

diff --git a/Engine/Plugins/Experimental/CommonUI/Source/CommonUI/Private/CommonGameViewportClient.cpp b/Engine/Plugins/Experimental/CommonUI/Source/CommonUI/Private/CommonGameViewportClient.cpp
--- a/Engine/Plugins/Experimental/CommonUI/Source/CommonUI/Private/CommonGameViewportClient.cpp
+++ b/Engine/Plugins/Experimental/CommonUI/Source/CommonUI/Private/CommonGameViewportClient.cpp
@@ -31,13 +31,22 @@ UCommonGameViewportClient::~UCommonGameViewportClient()
 
 bool UCommonGameViewportClient::InputKey(const FInputKeyEventArgs& InEventArgs)
 {
-	const FInputKeyEventArgs& EventArgs = InEventArgs;
+	FInputKeyEventArgs EventArgs = InEventArgs;
 
 	if (IsKeyPriorityAboveUI(EventArgs))
 	{
 		return true;
 	}
 
+	// Check override before UI
+	if (OnOverrideInputKey().IsBound())
+	{
+		if (OnOverrideInputKey().Execute(EventArgs))
+		{
+			return true;
+		}
+	}
+
 	// The input is fair game for handling - the UI gets first dibs
 #if !UE_BUILD_SHIPPING
 	if (ViewportConsole && !ViewportConsole->ConsoleState.IsEqual(NAME_Typing) && !ViewportConsole->ConsoleState.IsEqual(NAME_Open))
