Common Style Sheet:
- Add comments to emphasize that this is a prototype only
- Put the property in the advanced display

diff --git a/Engine/Plugins/Experimental/CommonUI/Source/CommonUI/Private/StyleSheet/CommonStyleSheet.cpp b/Engine/Plugins/Experimental/CommonUI/Source/CommonUI/Private/StyleSheet/CommonStyleSheet.cpp
--- a/Engine/Plugins/Experimental/CommonUI/Source/CommonUI/Private/StyleSheet/CommonStyleSheet.cpp
+++ b/Engine/Plugins/Experimental/CommonUI/Source/CommonUI/Private/StyleSheet/CommonStyleSheet.cpp
@@ -7,6 +7,10 @@
 #include "CommonTextBlock.h"
 #include "Components/Widget.h"
 
+//////////////////////////////
+// PROTOTYPE: DO NOT USE!!!
+//////////////////////////////
+
 void UCommonStyleSheet::Apply(UWidget* Widget)
 {
 	// FIXME: Current usage of the SetX functions can results in Apply being called
diff --git a/Engine/Plugins/Experimental/CommonUI/Source/CommonUI/Public/CommonTextBlock.h b/Engine/Plugins/Experimental/CommonUI/Source/CommonUI/Public/CommonTextBlock.h
--- a/Engine/Plugins/Experimental/CommonUI/Source/CommonUI/Public/CommonTextBlock.h
+++ b/Engine/Plugins/Experimental/CommonUI/Source/CommonUI/Public/CommonTextBlock.h
@@ -219,8 +219,8 @@ private:
 	UPROPERTY(EditAnywhere, BlueprintReadOnly, Category = CommonText, meta = (ExposeOnSpawn = true, AllowPrivateAccess = true))
 	TSubclassOf<UCommonTextScrollStyle> ScrollStyle;
 
-	/** References the scroll style asset to use, no reference disables scrolling*/
-	UPROPERTY(EditAnywhere, BlueprintReadOnly, Category = CommonText, meta = (ExposeOnSpawn = true, AllowPrivateAccess = true))
+	/** Prototype: Do not use! */
+	UPROPERTY(EditAnywhere, BlueprintReadOnly, AdvancedDisplay, Category = CommonText, meta = (ExposeOnSpawn = true, AllowPrivateAccess = true))
 	TObjectPtr<UCommonStyleSheet> StyleSheet;
 
 	/** If scrolling is enabled/disabled initially, this can be updated in blueprint */
diff --git a/Engine/Plugins/Experimental/CommonUI/Source/CommonUI/Public/StyleSheet/CommonStyleSheet.h b/Engine/Plugins/Experimental/CommonUI/Source/CommonUI/Public/StyleSheet/CommonStyleSheet.h
--- a/Engine/Plugins/Experimental/CommonUI/Source/CommonUI/Public/StyleSheet/CommonStyleSheet.h
+++ b/Engine/Plugins/Experimental/CommonUI/Source/CommonUI/Public/StyleSheet/CommonStyleSheet.h
@@ -10,6 +10,10 @@ class UCommonStyleSheetTypeBase;
 class UCommonTextBlock;
 class UWidget;
 
+//////////////////////////////
+// PROTOTYPE: DO NOT USE!!!
+//////////////////////////////
+
 UCLASS(BlueprintType)
 class UCommonStyleSheet : public UDataAsset
 {
