Fix missing broadcasts CommonLazyImage & CommonLazyWidget

diff --git a/Engine/Plugins/Experimental/CommonUI/Source/CommonUI/Private/CommonLazyImage.cpp b/Engine/Plugins/Experimental/CommonUI/Source/CommonUI/Private/CommonLazyImage.cpp
--- a/Engine/Plugins/Experimental/CommonUI/Source/CommonUI/Private/CommonLazyImage.cpp
+++ b/Engine/Plugins/Experimental/CommonUI/Source/CommonUI/Private/CommonLazyImage.cpp
@@ -200,6 +200,7 @@ void UCommonLazyImage::SetIsLoading(bool bIsLoading)
 
 void UCommonLazyImage::HandleLoadGuardStateChanged(bool bIsLoading)
 {
+	OnLoadingStateChangedEvent.Broadcast(bIsLoading);
 	BP_OnLoadingStateChanged.Broadcast(bIsLoading);
 }
 
diff --git a/Engine/Plugins/Experimental/CommonUI/Source/CommonUI/Private/CommonLazyWidget.cpp b/Engine/Plugins/Experimental/CommonUI/Source/CommonUI/Private/CommonLazyWidget.cpp
--- a/Engine/Plugins/Experimental/CommonUI/Source/CommonUI/Private/CommonLazyWidget.cpp
+++ b/Engine/Plugins/Experimental/CommonUI/Source/CommonUI/Private/CommonLazyWidget.cpp
@@ -186,6 +186,7 @@ const FText UCommonLazyWidget::GetPaletteCategory()
 
 void UCommonLazyWidget::HandleLoadGuardStateChanged(bool bIsLoading)
 {
+	OnLoadingStateChangedEvent.Broadcast(bIsLoading);
 	BP_OnLoadingStateChanged.Broadcast(bIsLoading);
 }
 
