PR #9477: Export FCommonAnalogCursor from the CommonUI module (Contributed by [REDACTED_PERSON])

diff --git a/Engine/Plugins/Experimental/CommonUI/Source/CommonUI/Public/Input/CommonAnalogCursor.h b/Engine/Plugins/Experimental/CommonUI/Source/CommonUI/Public/Input/CommonAnalogCursor.h
--- a/Engine/Plugins/Experimental/CommonUI/Source/CommonUI/Public/Input/CommonAnalogCursor.h
+++ b/Engine/Plugins/Experimental/CommonUI/Source/CommonUI/Public/Input/CommonAnalogCursor.h
@@ -22,7 +22,7 @@ enum EOrientation;
  * Introduces a separate focus-driven mode of operation, wherein the cursor is made invisible and automatically updated
  * to be centered over whatever widget is currently focused (except the game viewport - we completely hide it then)
  */
-class FCommonAnalogCursor : public FAnalogCursor
+class COMMONUI_API FCommonAnalogCursor : public FAnalogCursor
 {
 public:
 	template <typename AnalogCursorT = FCommonAnalogCursor>
@@ -37,7 +37,7 @@ public:
 	virtual bool HandleKeyDownEvent(FSlateApplication& SlateApp, const FKeyEvent& InKeyEvent) override;
 	virtual bool HandleKeyUpEvent(FSlateApplication& SlateApp, const FKeyEvent& InKeyEvent) override;
 
-	bool CanReleaseMouseCapture() const;
+	virtual bool CanReleaseMouseCapture() const;
 
 	virtual bool HandleAnalogInputEvent(FSlateApplication& SlateApp, const FAnalogInputEvent& InAnalogInputEvent) override;
 	virtual bool HandleMouseMoveEvent(FSlateApplication& SlateApp, const FPointerEvent& MouseEvent) override;
@@ -50,7 +50,7 @@ public:
 
 	virtual void ShouldHandleRightAnalog(bool bInShouldHandleRightAnalog);
 
-	bool IsAnalogMovementEnabled() const { return bIsAnalogMovementEnabled; }
+	virtual bool IsAnalogMovementEnabled() const { return bIsAnalogMovementEnabled; }
 
 protected:
 	FCommonAnalogCursor(const UCommonUIActionRouterBase& InActionRouter);
@@ -76,17 +76,23 @@ protected:
 	 */
 	bool IsGameViewportInFocusPathWithoutCapture() const;
 
-private:
-	//EAnalogStick GetScrollStick() const;
+	virtual void RefreshCursorSettings();
+	virtual void RefreshCursorVisibility();
+	
+	virtual void HandleInputMethodChanged(ECommonInputType NewInputMethod);
 	
-	void RefreshCursorSettings();
-	void RefreshCursorVisibility();
 	bool IsUsingGamepad() const;
 
-	void HandleInputMethodChanged(ECommonInputType NewInputMethod);
-
 	// Knowingly unorthodox member reference to a UObject - ok because we are a subobject of the owning router and will never outlive it
 	const UCommonUIActionRouterBase& ActionRouter;
+	ECommonInputType ActiveInputMethod;
+
+	bool bIsAnalogMovementEnabled = false;
+	
+	bool bShouldHandleRightAnalog = true;
+	
+private:
+	//EAnalogStick GetScrollStick() const;
 	
 	/** The current set of pointer buttons being used as keys. */
 	TSet<FKey> PointerButtonDownKeys;
@@ -95,10 +101,6 @@ private:
 	FSlateRenderTransform LastCursorTargetTransform;
 
 	float TimeUntilScrollUpdate = 0.f;
-	ECommonInputType ActiveInputMethod;
-	bool bIsAnalogMovementEnabled = false;
-
-	bool bShouldHandleRightAnalog = true;
 
 	//@todo DanH: This actually needs to go up a layer into FAnalogCursor so it behaves as desired if we allow things to fall through
 	//		We'll also need a CommonNavigationConfig that respects the stick assignment as well so it also maps to nav mode
@@ -116,4 +118,4 @@ private:
 	};
 	int32 ShoulderButtonStatus = EShoulderButtonFlags::None;
 #endif
-};
\ No newline at end of file
+};
