Adding documentation to the CommonActionWidget to explain what the material parameters are for, and the expected values for the progress material.

diff --git a/Engine/Plugins/Experimental/CommonUI/Source/CommonUI/Public/CommonActionWidget.h b/Engine/Plugins/Experimental/CommonUI/Source/CommonUI/Public/CommonActionWidget.h
--- a/Engine/Plugins/Experimental/CommonUI/Source/CommonUI/Public/CommonActionWidget.h
+++ b/Engine/Plugins/Experimental/CommonUI/Source/CommonUI/Public/CommonActionWidget.h
@@ -55,9 +55,14 @@ public:
 	UPROPERTY(BlueprintAssignable, Category = CommonActionWidget)
 	FOnInputMethodChanged OnInputMethodChanged;
 
+	/**
+	 * The material to use when showing held progress, the progress will be sent using the material parameter
+	 * defined by ProgressMaterialParam and the value will range from 0..1.
+	 **/
 	UPROPERTY(EditAnywhere, BlueprintReadOnly, Category = CommonActionWidget)
 	FSlateBrush ProgressMaterialBrush;
 
+	/** The material parameter on ProgressMaterialBrush to update the held percentage.  This value will be 0..1. */
 	UPROPERTY(EditAnywhere, BlueprintReadOnly, Category = CommonActionWidget)
 	FName ProgressMaterialParam;
 
