[CommonUI] Added BP_OnPressed() and BP_OnReleased() blueprint-implementable events to CommonButtonBase for better control over visual states and animations.

diff --git a/Engine/Plugins/Experimental/CommonUI/Source/CommonUI/Private/CommonButtonBase.cpp b/Engine/Plugins/Experimental/CommonUI/Source/CommonUI/Private/CommonButtonBase.cpp
--- a/Engine/Plugins/Experimental/CommonUI/Source/CommonUI/Private/CommonButtonBase.cpp
+++ b/Engine/Plugins/Experimental/CommonUI/Source/CommonUI/Private/CommonButtonBase.cpp
@@ -1046,7 +1046,7 @@ void UCommonButtonBase::HandleFocusReceived()
 
 void UCommonButtonBase::HandleButtonPressed()
 {
-	OnPressed().Broadcast();
+	NativeOnPressed();
 
 	UCommonInputSubsystem* CommonInputSubsystem = GetInputSubsystem();
 
@@ -1059,7 +1059,7 @@ void UCommonButtonBase::HandleButtonPressed()
 
 void UCommonButtonBase::HandleButtonReleased()
 {
-	OnReleased().Broadcast();
+	NativeOnReleased();
 
 	UCommonInputSubsystem* CommonInputSubsystem = GetInputSubsystem();
 
@@ -1170,6 +1170,18 @@ void UCommonButtonBase::StopDoubleClickPropagation()
 	bStopDoubleClickPropagation = true;
 }
 
+void UCommonButtonBase::NativeOnPressed()
+{
+	BP_OnPressed();
+	OnPressed().Broadcast();
+}
+
+void UCommonButtonBase::NativeOnReleased()
+{
+	BP_OnReleased();
+	OnReleased().Broadcast();
+}
+
 void UCommonButtonBase::NativeOnEnabled()
 {
 	BP_OnEnabled();
diff --git a/Engine/Plugins/Experimental/CommonUI/Source/CommonUI/Public/CommonButtonBase.h b/Engine/Plugins/Experimental/CommonUI/Source/CommonUI/Public/CommonButtonBase.h
--- a/Engine/Plugins/Experimental/CommonUI/Source/CommonUI/Public/CommonButtonBase.h
+++ b/Engine/Plugins/Experimental/CommonUI/Source/CommonUI/Public/CommonButtonBase.h
@@ -520,6 +520,14 @@ protected:
 	/** Unless this is called, we will assume the double click should be converted into a normal click. */
 	UFUNCTION(BlueprintCallable, Category = CommonButton)
 	void StopDoubleClickPropagation();
+	
+	UFUNCTION(BlueprintImplementableEvent, Category = CommonButton, meta = (DisplayName = "On Pressed"))
+	void BP_OnPressed();
+	virtual void NativeOnPressed();
+
+	UFUNCTION(BlueprintImplementableEvent, Category = CommonButton, meta = (DisplayName = "On Released"))
+	void BP_OnReleased();
+	virtual void NativeOnReleased();
 
 	UFUNCTION(BlueprintImplementableEvent, Category = CommonButton, meta = (DisplayName = "On Enabled"))
 	void BP_OnEnabled();
