Modify CommonButtonBase to handle new features:
- Added Lockable state to button, with BP Implementable events when Locked changes to handle visuals.
- Added OnFocusLost delegate and BP Implementable events to handle visuals.
- Added an Invisible Button styling for the underlying Slate Button.

diff --git a/Engine/Plugins/Experimental/CommonUI/Source/CommonUI/Private/CommonButtonBase.cpp b/Engine/Plugins/Experimental/CommonUI/Source/CommonUI/Private/CommonButtonBase.cpp
--- a/Engine/Plugins/Experimental/CommonUI/Source/CommonUI/Private/CommonButtonBase.cpp
+++ b/Engine/Plugins/Experimental/CommonUI/Source/CommonUI/Private/CommonButtonBase.cpp
@@ -237,7 +237,8 @@ TSharedRef<SWidget> UCommonButtonInternalBase::RebuildWidget()
 		.IsFocusable(IsFocusable)
 		.IsButtonEnabled(bButtonEnabled)
 		.IsInteractionEnabled(bInteractionEnabled)
-		.OnReceivedFocus(BIND_UOBJECT_DELEGATE(FSimpleDelegate, SlateHandleOnReceivedFocus));
+		.OnReceivedFocus(BIND_UOBJECT_DELEGATE(FSimpleDelegate, SlateHandleOnReceivedFocus))
+		.OnLostFocus(BIND_UOBJECT_DELEGATE(FSimpleDelegate, SlateHandleOnLostFocus));
 
 	MyBox = SNew(SBox)
 		.MinDesiredWidth(MinWidth)
@@ -300,6 +301,11 @@ void UCommonButtonInternalBase::SlateHandleOnReceivedFocus()
 	OnReceivedFocus.ExecuteIfBound();
 }
 
+void UCommonButtonInternalBase::SlateHandleOnLostFocus()
+{
+	OnLostFocus.ExecuteIfBound();
+}
+
 //////////////////////////////////////////////////////////////////////////
 // UCommonButtonBase
 //////////////////////////////////////////////////////////////////////////
@@ -316,6 +322,7 @@ UCommonButtonBase::UCommonButtonBase(const FObjectInitializer& ObjectInitializer
 	, bDisplayInputActionWhenNotInteractable(true)
 	, bShouldUseFallbackDefaultInputAction(true)
 	, bSelected(false)
+	, bLocked(false)
 	, bButtonEnabled(true)
 	, bInteractionEnabled(true)
 {
@@ -397,6 +404,7 @@ bool UCommonButtonBase::Initialize()
 			RootButton->OnClicked.AddUniqueDynamic(this, &UCommonButtonBase::HandleButtonClicked);
 			RootButton->HandleDoubleClicked.BindUObject(this, &UCommonButtonBase::HandleButtonDoubleClicked);
 			RootButton->OnReceivedFocus.BindUObject(this, &UCommonButtonBase::HandleFocusReceived);
+			RootButton->OnLostFocus.BindUObject(this, &UCommonButtonBase::HandleFocusLost);
 			RootButton->OnPressed.AddUniqueDynamic(this, &UCommonButtonBase::HandleButtonPressed);
 			RootButton->OnReleased.AddUniqueDynamic(this, &UCommonButtonBase::HandleButtonReleased);
 		}
@@ -721,6 +729,13 @@ void UCommonButtonBase::SetIsSelected(bool InSelected, bool bGiveClickFeedback)
 	}
 }
 
+void UCommonButtonBase::SetIsLocked(bool bInIsLocked)
+{
+	bLocked = bInIsLocked;
+
+	BP_OnLockedChanged(bLocked);
+}
+
 void UCommonButtonBase::SetSelectedInternal(bool bInSelected, bool bAllowSound /*= true*/, bool bBroadcast /*= true*/)
 {
 	bSelected = bInSelected;
@@ -795,6 +810,11 @@ bool UCommonButtonBase::GetSelected() const
 	return bSelected;
 }
 
+bool UCommonButtonBase::GetLocked() const
+{
+	return bLocked;
+}
+
 void UCommonButtonBase::ClearSelection()
 {
 	SetSelectedInternal( false, false );
@@ -1042,6 +1062,13 @@ void UCommonButtonBase::HandleFocusReceived()
 		SetIsSelected(true, false);
 	}
 	OnFocusReceived().Broadcast();
+	BP_OnFocusReceived();
+}
+
+void UCommonButtonBase::HandleFocusLost()
+{
+	OnFocusLost().Broadcast();
+	BP_OnFocusLost();
 }
 
 void UCommonButtonBase::HandleButtonPressed()
@@ -1079,6 +1106,11 @@ FReply UCommonButtonBase::NativeOnFocusReceived(const FGeometry& InGeometry, con
 	return Reply;
 }
 
+void UCommonButtonBase::NativeOnFocusLost(const FFocusEvent& InFocusEvent)
+{
+	Super::NativeOnFocusLost(InFocusEvent);
+}
+
 void UCommonButtonBase::NativeOnSelected(bool bBroadcast)
 {
 	BP_OnSelected();
@@ -1135,33 +1167,49 @@ void UCommonButtonBase::NativeOnUnhovered()
 
 void UCommonButtonBase::NativeOnClicked()
 {
-	BP_OnClicked();
-	OnClicked().Broadcast();
-	if (OnButtonBaseClicked.IsBound())
+	if(!GetLocked())
 	{
-		OnButtonBaseClicked.Broadcast(this);
-	}
-
-	FString ButtonName, ABTestName, ExtraData;
-	if (GetButtonAnalyticInfo(ButtonName, ABTestName, ExtraData))
-	{
-		UCommonUISubsystemBase* CommonUISubsystem = GetUISubsystem();
-		if (GetGameInstance())
+		BP_OnClicked();
+		OnClicked().Broadcast();
+		if (OnButtonBaseClicked.IsBound())
+		{
+			OnButtonBaseClicked.Broadcast(this);
+		}
+		
+		FString ButtonName, ABTestName, ExtraData;
+		if (GetButtonAnalyticInfo(ButtonName, ABTestName, ExtraData))
 		{
-			check(CommonUISubsystem);
+			UCommonUISubsystemBase* CommonUISubsystem = GetUISubsystem();
+			if (GetGameInstance())
+			{
+				check(CommonUISubsystem);
 
-			CommonUISubsystem->FireEvent_ButtonClicked(ButtonName, ABTestName, ExtraData);
+				CommonUISubsystem->FireEvent_ButtonClicked(ButtonName, ABTestName, ExtraData);
+			}
 		}
 	}
+	else
+	{
+		BP_OnLockClicked();
+		OnLockClicked().Broadcast();
+	}
 }
 
 void UCommonButtonBase::NativeOnDoubleClicked()
 {
-	BP_OnDoubleClicked();
-	OnDoubleClicked().Broadcast();
-	if (OnButtonBaseDoubleClicked.IsBound())
+	if(!GetLocked())
+	{
+		BP_OnDoubleClicked();
+		OnDoubleClicked().Broadcast();
+		if (OnButtonBaseDoubleClicked.IsBound())
+		{
+			OnButtonBaseDoubleClicked.Broadcast(this);
+		}
+	}
+	else
 	{
-		OnButtonBaseDoubleClicked.Broadcast(this);
+		BP_OnLockDoubleClicked();
+		OnLockDoubleClicked().Broadcast();
 	}
 }
 
diff --git a/Engine/Plugins/Experimental/CommonUI/Source/CommonUI/Private/CommonButtonTypes.cpp b/Engine/Plugins/Experimental/CommonUI/Source/CommonUI/Private/CommonButtonTypes.cpp
--- a/Engine/Plugins/Experimental/CommonUI/Source/CommonUI/Private/CommonButtonTypes.cpp
+++ b/Engine/Plugins/Experimental/CommonUI/Source/CommonUI/Private/CommonButtonTypes.cpp
@@ -196,6 +196,11 @@ FReply SCommonButton::OnFocusReceived(const FGeometry& MyGeometry, const FFocusE
 	return ReturnReply;
 }
 
+void SCommonButton::OnFocusLost(const FFocusEvent& InFocusEvent)
+{
+	OnLostFocus.ExecuteIfBound();
+}
+
 int32 SCommonButton::OnPaint(const FPaintArgs& Args, const FGeometry& AllottedGeometry, const FSlateRect& MyClippingRect, FSlateWindowElementList& OutDrawElements, int32 LayerId, const FWidgetStyle& InWidgetStyle, bool bParentEnabled) const
 {
 	bool bEnabled = bParentEnabled && bIsButtonEnabled;
diff --git a/Engine/Plugins/Experimental/CommonUI/Source/CommonUI/Private/CommonButtonTypes.h b/Engine/Plugins/Experimental/CommonUI/Source/CommonUI/Private/CommonButtonTypes.h
--- a/Engine/Plugins/Experimental/CommonUI/Source/CommonUI/Private/CommonButtonTypes.h
+++ b/Engine/Plugins/Experimental/CommonUI/Source/CommonUI/Private/CommonButtonTypes.h
@@ -38,6 +38,7 @@ public:
 		SLATE_ARGUMENT(EButtonPressMethod::Type, PressMethod)
 		SLATE_ARGUMENT(bool, IsFocusable)
 		SLATE_EVENT(FSimpleDelegate, OnReceivedFocus)
+		SLATE_EVENT(FSimpleDelegate, OnLostFocus)
 
 		/** Is interaction enabled? */
 		SLATE_ARGUMENT(bool, IsButtonEnabled)
@@ -69,6 +70,7 @@ public:
 		SetHover(false);
 
 		OnReceivedFocus = InArgs._OnReceivedFocus;
+		OnLostFocus = InArgs._OnLostFocus;
 		bIsButtonEnabled = InArgs._IsButtonEnabled;
 		bIsInteractionEnabled = InArgs._IsInteractionEnabled;
 		bHovered = false;
@@ -101,6 +103,8 @@ public:
 	/** Overridden to fire delegate for external listener */
 	virtual FReply OnFocusReceived(const FGeometry& MyGeometry, const FFocusEvent& InFocusEvent);
 
+	virtual void OnFocusLost(const FFocusEvent& InFocusEvent) override;
+
 	virtual int32 OnPaint(const FPaintArgs& Args, const FGeometry& AllottedGeometry, const FSlateRect& MyClippingRect, FSlateWindowElementList& OutDrawElements, int32 LayerId, const FWidgetStyle& InWidgetStyle, bool bParentEnabled) const override;
 
 protected:
@@ -113,6 +117,9 @@ private:
 	/** Delegate fired whenever focus is received */
 	FSimpleDelegate OnReceivedFocus;
 
+	/** Delegate fired whenever focus is lost */
+	FSimpleDelegate OnLostFocus;
+
 	/** True if the button is enabled */
 	bool bIsButtonEnabled;
 
diff --git a/Engine/Plugins/Experimental/CommonUI/Source/CommonUI/Public/CommonButtonBase.h b/Engine/Plugins/Experimental/CommonUI/Source/CommonUI/Public/CommonButtonBase.h
--- a/Engine/Plugins/Experimental/CommonUI/Source/CommonUI/Public/CommonButtonBase.h
+++ b/Engine/Plugins/Experimental/CommonUI/Source/CommonUI/Public/CommonButtonBase.h
@@ -230,6 +230,9 @@ public:
 	/** Called when the button receives focus */
 	FSimpleDelegate OnReceivedFocus;
 
+	/** Called when the button loses focus */
+	FSimpleDelegate OnLostFocus;
+
 protected:
 	// UWidget interface
 	virtual TSharedRef<SWidget> RebuildWidget() override;
@@ -244,6 +247,9 @@ protected:
 	/** Called when internal slate button receives focus; Fires OnReceivedFocus */
 	void SlateHandleOnReceivedFocus();
 
+	/** Called when internal slate button loses focus; Fires OnLostFocus */
+	void SlateHandleOnLostFocus();
+
 	/** The minimum width of the button */
 	UPROPERTY()
 	int32 MinWidth;
@@ -346,10 +352,18 @@ public:
 	UFUNCTION(BlueprintCallable, Category = "Common Button|Setters")
 	void SetIsSelected(bool InSelected, bool bGiveClickFeedback = true);
 
+	/** Change whether this widget is locked. If locked, the button can be focusable and responsive to mouse input but will not broadcast OnClicked events. */
+	UFUNCTION(BlueprintCallable, Category = "Common Button|Setters")
+	void SetIsLocked(bool bInIsLocked);
+
 	/** @returns True if the button is currently in a selected state, False otherwise */
 	UFUNCTION(BlueprintCallable, Category = "Common Button|Getters")
 	bool GetSelected() const;
 
+	/** @returns True if the button is currently locked, False otherwise */
+	UFUNCTION(BlueprintCallable, Category = "Common Button|Getters")
+	bool GetLocked() const;
+
 	UFUNCTION(BlueprintCallable, Category = "Common Button" )
 	void ClearSelection();
 
@@ -432,6 +446,9 @@ public:
 	FCommonButtonEvent& OnHovered() const { return OnHoveredEvent; }
 	FCommonButtonEvent& OnUnhovered() const { return OnUnhoveredEvent; }
 	FCommonButtonEvent& OnFocusReceived() const { return OnFocusReceivedEvent; }
+	FCommonButtonEvent& OnFocusLost() const { return OnFocusLostEvent; }
+	FCommonButtonEvent& OnLockClicked() const { return OnLockClickedEvent; }
+	FCommonButtonEvent& OnLockDoubleClicked() const { return OnLockDoubleClickedEvent; }
 
 	DECLARE_EVENT_OneParam(UCommonButtonBase, FOnIsSelectedChanged, bool);
 	FOnIsSelectedChanged& OnIsSelectedChanged() const { return OnIsSelectedChangedEvent; }
@@ -443,6 +460,7 @@ protected:
 	virtual void PostLoad() override;
 	virtual void SynchronizeProperties() override;
 	virtual FReply NativeOnFocusReceived(const FGeometry& InGeometry, const FFocusEvent& InFocusEvent) override;
+	virtual void NativeOnFocusLost(const FFocusEvent& InFocusEvent) override;
 
 #if WITH_EDITOR
 	virtual void OnCreationFromPalette() override;
@@ -484,6 +502,10 @@ protected:
 	/** Helper function registered to the underlying button receiving focus */
 	UFUNCTION()
 	virtual void HandleFocusReceived();
+	
+	/** Helper function registered to the underlying button losing focus */
+	UFUNCTION()
+	virtual void HandleFocusLost();
 
 	/** Helper function registered to the underlying button when pressed */
 	UFUNCTION()
@@ -508,6 +530,21 @@ protected:
 	UFUNCTION(BlueprintImplementableEvent, Category = CommonButton, meta = (DisplayName = "On Unhovered"))
 	void BP_OnUnhovered();
 	virtual void NativeOnUnhovered();
+	
+	UFUNCTION(BlueprintImplementableEvent, Category = CommonButton, meta = (DisplayName = "On Focused"))
+	void BP_OnFocusReceived();
+
+	UFUNCTION(BlueprintImplementableEvent, Category = CommonButton, meta = (DisplayName = "On Unfocused"))
+	void BP_OnFocusLost();
+
+	UFUNCTION(BlueprintImplementableEvent, Category = CommonButton, meta = (DisplayName = "On Locked Changed"))
+	void BP_OnLockedChanged(bool bIsLocked);
+	
+	UFUNCTION(BlueprintImplementableEvent, Category = CommonButton, meta = (DisplayName = "On Lock Clicked"))
+	void BP_OnLockClicked();
+
+	UFUNCTION(BlueprintImplementableEvent, Category = CommonButton, meta = (DisplayName = "On Lock Double Clicked"))
+	void BP_OnLockDoubleClicked();
 
 	UFUNCTION(BlueprintImplementableEvent, Category = CommonButton, meta = (DisplayName = "On Clicked"))
 	void BP_OnClicked();
@@ -613,7 +650,7 @@ protected:
 	/** True if the button can be deselected by clicking it when selected */
 	UPROPERTY(EditAnywhere, BlueprintReadOnly, Category = Selection, meta = (ExposeOnSpawn = true, EditCondition = "bSelectable"))
 	uint8 bToggleable:1;
-
+	
 	UPROPERTY(EditAnywhere, BlueprintReadOnly, Category = Selection, meta = (ExposeOnSpawn = true, EditCondition = "bSelectable"))
 	uint8 bTriggerClickedAfterSelection:1;
 
@@ -634,6 +671,9 @@ private:
 	/** True if this button is currently selected */
 	uint8 bSelected:1;
 
+	/** True if this button is currently locked */
+	uint8 bLocked:1;
+
 	/** True if this button is currently enabled */
 	uint8 bButtonEnabled:1;
 
@@ -754,6 +794,9 @@ private:
 	mutable FCommonButtonEvent OnHoveredEvent;
 	mutable FCommonButtonEvent OnUnhoveredEvent;
 	mutable FCommonButtonEvent OnFocusReceivedEvent;
+	mutable FCommonButtonEvent OnFocusLostEvent;
+	mutable FCommonButtonEvent OnLockClickedEvent;
+	mutable FCommonButtonEvent OnLockDoubleClickedEvent;
 
 	mutable FOnIsSelectedChanged OnIsSelectedChangedEvent;
 
