Add a BP Implementable event that the artist may use to change visuals whenever the Input Method changes.

diff --git a/Engine/Plugins/Experimental/CommonUI/Source/CommonUI/Private/CommonButtonBase.cpp b/Engine/Plugins/Experimental/CommonUI/Source/CommonUI/Private/CommonButtonBase.cpp
--- a/Engine/Plugins/Experimental/CommonUI/Source/CommonUI/Private/CommonButtonBase.cpp
+++ b/Engine/Plugins/Experimental/CommonUI/Source/CommonUI/Private/CommonButtonBase.cpp
@@ -478,6 +478,7 @@ void UCommonButtonBase::UnbindInputMethodChangedDelegate()
 void UCommonButtonBase::OnInputMethodChanged(ECommonInputType CurrentInputType)
 {
 	UpdateInputActionWidget();
+	BP_OnInputMethodChanged(CurrentInputType);
 }
 
 void UCommonButtonBase::BindTriggeringInputActionToClick()
diff --git a/Engine/Plugins/Experimental/CommonUI/Source/CommonUI/Public/CommonButtonBase.h b/Engine/Plugins/Experimental/CommonUI/Source/CommonUI/Public/CommonButtonBase.h
--- a/Engine/Plugins/Experimental/CommonUI/Source/CommonUI/Public/CommonButtonBase.h
+++ b/Engine/Plugins/Experimental/CommonUI/Source/CommonUI/Public/CommonButtonBase.h
@@ -574,6 +574,9 @@ protected:
 	void BP_OnDisabled();
 	virtual void NativeOnDisabled();
 
+	UFUNCTION(BlueprintImplementableEvent, Category = CommonButton, meta = (DisplayName = "On Input Method Changed"))
+	void BP_OnInputMethodChanged(ECommonInputType CurrentInputType);
+
 	/** Allows derived classes to take action when the current text style has changed */
 	UFUNCTION(BlueprintImplementableEvent, meta=(BlueprintProtected="true"), Category = "Common Button")
 	void OnCurrentTextStyleChanged();
