Allow runtime edits to the Progress Brush of CommonActionWidget.h in blueprints

diff --git a/Engine/Plugins/Experimental/CommonUI/Source/CommonUI/Public/CommonActionWidget.h b/Engine/Plugins/Experimental/CommonUI/Source/CommonUI/Public/CommonActionWidget.h
--- a/Engine/Plugins/Experimental/CommonUI/Source/CommonUI/Public/CommonActionWidget.h
+++ b/Engine/Plugins/Experimental/CommonUI/Source/CommonUI/Public/CommonActionWidget.h
@@ -59,7 +59,7 @@ public:
 	 * The material to use when showing held progress, the progress will be sent using the material parameter
 	 * defined by ProgressMaterialParam and the value will range from 0..1.
 	 **/
-	UPROPERTY(EditAnywhere, BlueprintReadOnly, Category = CommonActionWidget)
+	UPROPERTY(EditAnywhere, BlueprintReadWrite, Category = CommonActionWidget)
 	FSlateBrush ProgressMaterialBrush;
 
 	/** The material parameter on ProgressMaterialBrush to update the held percentage.  This value will be 0..1. */
