Mobile platforms should have UE_COMMONINPUT_PLATFORM_SUPPORTS_TOUCH set true

#android
#ios

diff --git a/Engine/Plugins/Experimental/CommonUI/Source/CommonInput/CommonInput.Build.cs b/Engine/Plugins/Experimental/CommonUI/Source/CommonInput/CommonInput.Build.cs
--- a/Engine/Plugins/Experimental/CommonUI/Source/CommonInput/CommonInput.Build.cs
+++ b/Engine/Plugins/Experimental/CommonUI/Source/CommonInput/CommonInput.Build.cs
@@ -86,7 +86,10 @@ public class CommonInput : ModuleRules
 		get
 		{
 			// Support touch testing until touch is supported on desktop
-			return Target.Platform.IsInGroup(UnrealPlatformGroup.Desktop) && Target.Configuration != UnrealTargetConfiguration.Shipping;
+			return (Target.Platform.IsInGroup(UnrealPlatformGroup.Desktop) && Target.Configuration != UnrealTargetConfiguration.Shipping)
+					// Platforms always with touch support
+					|| (Target.Platform == UnrealTargetPlatform.Android)
+					|| (Target.Platform == UnrealTargetPlatform.IOS);
 		}
 	}
 }
diff --git a/Engine/Plugins/Experimental/CommonUI/Source/CommonUI/CommonUI.Build.cs b/Engine/Plugins/Experimental/CommonUI/Source/CommonUI/CommonUI.Build.cs
--- a/Engine/Plugins/Experimental/CommonUI/Source/CommonUI/CommonUI.Build.cs
+++ b/Engine/Plugins/Experimental/CommonUI/Source/CommonUI/CommonUI.Build.cs
@@ -67,7 +67,10 @@ public class CommonUI : ModuleRules
 		get
 		{
 			// Support touch testing until touch is supported on desktop
-			return Target.Platform.IsInGroup(UnrealPlatformGroup.Desktop) && Target.Configuration != UnrealTargetConfiguration.Shipping;
+			return (Target.Platform.IsInGroup(UnrealPlatformGroup.Desktop) && Target.Configuration != UnrealTargetConfiguration.Shipping)
+					// Platforms always with touch support
+					|| (Target.Platform == UnrealTargetPlatform.Android)
+					|| (Target.Platform == UnrealTargetPlatform.IOS);
 		}
 	}
 }
