Renamed private pch header in CommonUI and CommonInput and fixed -NoPCH -DisableUnity compile issues.

diff --git a/Engine/Plugins/Experimental/CommonUI/Source/CommonInput/Private/CommonInputModule.cpp b/Engine/Plugins/Experimental/CommonUI/Source/CommonInput/Private/CommonInputModule.cpp
--- a/Engine/Plugins/Experimental/CommonUI/Source/CommonInput/Private/CommonInputModule.cpp
+++ b/Engine/Plugins/Experimental/CommonUI/Source/CommonInput/Private/CommonInputModule.cpp
@@ -1,6 +1,6 @@
 // Copyright Epic Games, Inc. All Rights Reserved.
 
-#include "CommonInputPrivatePCH.h"
+#include "CommonInputPrivate.h"
 #include "ICommonInputModule.h"
 #include "Engine/AssetManager.h"
 #if WITH_EDITOR
diff --git a/Engine/Plugins/Experimental/CommonUI/Source/CommonInput/Private/CommonInputPrivate.h b/Engine/Plugins/Experimental/CommonUI/Source/CommonInput/Private/CommonInputPrivate.h
new file mode 100644
--- /dev/null
+++ b/Engine/Plugins/Experimental/CommonUI/Source/CommonInput/Private/CommonInputPrivate.h
@@ -0,0 +1,11 @@
+// Copyright Epic Games, Inc. All Rights Reserved.
+
+#pragma once
+
+
+/* Dependencies
+ *****************************************************************************/
+
+#include "Styling/SlateStyle.h"
+
+DECLARE_LOG_CATEGORY_EXTERN(LogCommonInput, Log, All);
\ No newline at end of file
diff --git a/Engine/Plugins/Experimental/CommonUI/Source/CommonInput/Private/CommonInputPrivatePCH.h b/Engine/Plugins/Experimental/CommonUI/Source/CommonInput/Private/CommonInputPrivatePCH.h
deleted file mode 100644
diff --git a/Engine/Plugins/Experimental/CommonUI/Source/CommonInput/Private/CommonInputSettings.cpp b/Engine/Plugins/Experimental/CommonUI/Source/CommonInput/Private/CommonInputSettings.cpp
--- a/Engine/Plugins/Experimental/CommonUI/Source/CommonInput/Private/CommonInputSettings.cpp
+++ b/Engine/Plugins/Experimental/CommonUI/Source/CommonInput/Private/CommonInputSettings.cpp
@@ -1,7 +1,7 @@
 // Copyright Epic Games, Inc. All Rights Reserved.
 
 #include "CommonInputSettings.h"
-#include "CommonInputPrivatePCH.h"
+#include "CommonInputPrivate.h"
 #include "Engine/UserInterfaceSettings.h"
 #include "Engine/StreamableManager.h"
 #include "Engine/AssetManager.h"
diff --git a/Engine/Plugins/Experimental/CommonUI/Source/CommonInput/Private/CommonInputSubsystem.cpp b/Engine/Plugins/Experimental/CommonUI/Source/CommonInput/Private/CommonInputSubsystem.cpp
--- a/Engine/Plugins/Experimental/CommonUI/Source/CommonInput/Private/CommonInputSubsystem.cpp
+++ b/Engine/Plugins/Experimental/CommonUI/Source/CommonInput/Private/CommonInputSubsystem.cpp
@@ -1,7 +1,7 @@
 // Copyright Epic Games, Inc. All Rights Reserved.
 
 #include "CommonInputSubsystem.h"
-#include "CommonInputPrivatePCH.h"
+#include "CommonInputPrivate.h"
 #include "Misc/ConfigCacheIni.h"
 
 #include "Framework/Application/IInputProcessor.h"
diff --git a/Engine/Plugins/Experimental/CommonUI/Source/CommonUI/Private/CommonActionWidget.cpp b/Engine/Plugins/Experimental/CommonUI/Source/CommonUI/Private/CommonActionWidget.cpp
--- a/Engine/Plugins/Experimental/CommonUI/Source/CommonUI/Private/CommonActionWidget.cpp
+++ b/Engine/Plugins/Experimental/CommonUI/Source/CommonUI/Private/CommonActionWidget.cpp
@@ -1,7 +1,7 @@
 // Copyright Epic Games, Inc. All Rights Reserved.
 
 #include "CommonActionWidget.h"
-#include "CommonUIPrivatePCH.h"
+#include "CommonUIPrivate.h"
 #include "CommonUISubsystemBase.h"
 #include "CommonInputSubsystem.h"
 #include "CommonWidgetPaletteCategories.h"
diff --git a/Engine/Plugins/Experimental/CommonUI/Source/CommonUI/Private/CommonActivatableWidget.cpp b/Engine/Plugins/Experimental/CommonUI/Source/CommonUI/Private/CommonActivatableWidget.cpp
--- a/Engine/Plugins/Experimental/CommonUI/Source/CommonUI/Private/CommonActivatableWidget.cpp
+++ b/Engine/Plugins/Experimental/CommonUI/Source/CommonUI/Private/CommonActivatableWidget.cpp
@@ -1,7 +1,7 @@
 // Copyright Epic Games, Inc. All Rights Reserved.
 
 #include "CommonActivatableWidget.h"
-#include "CommonUIPrivatePCH.h"
+#include "CommonUIPrivate.h"
 #include "Input/CommonUIActionRouterBase.h"
 #include "Input/CommonUIInputTypes.h"
 #include "Input/UIActionRouterTypes.h"
diff --git a/Engine/Plugins/Experimental/CommonUI/Source/CommonUI/Private/CommonBorder.cpp b/Engine/Plugins/Experimental/CommonUI/Source/CommonUI/Private/CommonBorder.cpp
--- a/Engine/Plugins/Experimental/CommonUI/Source/CommonUI/Private/CommonBorder.cpp
+++ b/Engine/Plugins/Experimental/CommonUI/Source/CommonUI/Private/CommonBorder.cpp
@@ -1,7 +1,7 @@
 // Copyright Epic Games, Inc. All Rights Reserved.
 
 #include "CommonBorder.h"
-#include "CommonUIPrivatePCH.h"
+#include "CommonUIPrivate.h"
 #include "CommonUIEditorSettings.h"
 #include "CommonWidgetPaletteCategories.h"
 
diff --git a/Engine/Plugins/Experimental/CommonUI/Source/CommonUI/Private/CommonButtonBase.cpp b/Engine/Plugins/Experimental/CommonUI/Source/CommonUI/Private/CommonButtonBase.cpp
--- a/Engine/Plugins/Experimental/CommonUI/Source/CommonUI/Private/CommonButtonBase.cpp
+++ b/Engine/Plugins/Experimental/CommonUI/Source/CommonUI/Private/CommonButtonBase.cpp
@@ -1,7 +1,7 @@
 // Copyright Epic Games, Inc. All Rights Reserved.
 
 #include "CommonButtonBase.h"
-#include "CommonUIPrivatePCH.h"
+#include "CommonUIPrivate.h"
 
 #include "CommonTextBlock.h"
 #include "CommonActionWidget.h"
diff --git a/Engine/Plugins/Experimental/CommonUI/Source/CommonUI/Private/CommonCustomNavigation.cpp b/Engine/Plugins/Experimental/CommonUI/Source/CommonUI/Private/CommonCustomNavigation.cpp
--- a/Engine/Plugins/Experimental/CommonUI/Source/CommonUI/Private/CommonCustomNavigation.cpp
+++ b/Engine/Plugins/Experimental/CommonUI/Source/CommonUI/Private/CommonCustomNavigation.cpp
@@ -1,7 +1,7 @@
 // Copyright Epic Games, Inc. All Rights Reserved.
 
 #include "CommonCustomNavigation.h"
-#include "CommonUIPrivatePCH.h"
+#include "CommonUIPrivate.h"
 #include "CommonUISettings.h"
 #include "CommonWidgetPaletteCategories.h"
 
diff --git a/Engine/Plugins/Experimental/CommonUI/Source/CommonUI/Private/CommonDateTimeTextBlock.cpp b/Engine/Plugins/Experimental/CommonUI/Source/CommonUI/Private/CommonDateTimeTextBlock.cpp
--- a/Engine/Plugins/Experimental/CommonUI/Source/CommonUI/Private/CommonDateTimeTextBlock.cpp
+++ b/Engine/Plugins/Experimental/CommonUI/Source/CommonUI/Private/CommonDateTimeTextBlock.cpp
@@ -1,7 +1,7 @@
 // Copyright Epic Games, Inc. All Rights Reserved.
 
 #include "CommonDateTimeTextBlock.h"
-#include "CommonUIPrivatePCH.h"
+#include "CommonUIPrivate.h"
 #include "CommonWidgetPaletteCategories.h"
 #include "TimerManager.h"
 #include "Widgets/Text/STextBlock.h"
diff --git a/Engine/Plugins/Experimental/CommonUI/Source/CommonUI/Private/CommonHardwareVisibilityBorder.cpp b/Engine/Plugins/Experimental/CommonUI/Source/CommonUI/Private/CommonHardwareVisibilityBorder.cpp
--- a/Engine/Plugins/Experimental/CommonUI/Source/CommonUI/Private/CommonHardwareVisibilityBorder.cpp
+++ b/Engine/Plugins/Experimental/CommonUI/Source/CommonUI/Private/CommonHardwareVisibilityBorder.cpp
@@ -1,7 +1,7 @@
 // Copyright Epic Games, Inc. All Rights Reserved.
 
 #include "CommonHardwareVisibilityBorder.h"
-#include "CommonUIPrivatePCH.h"
+#include "CommonUIPrivate.h"
 #include "CommonUISubsystemBase.h"
 
 #include "CommonInputBaseTypes.h"
diff --git a/Engine/Plugins/Experimental/CommonUI/Source/CommonUI/Private/CommonListView.cpp b/Engine/Plugins/Experimental/CommonUI/Source/CommonUI/Private/CommonListView.cpp
--- a/Engine/Plugins/Experimental/CommonUI/Source/CommonUI/Private/CommonListView.cpp
+++ b/Engine/Plugins/Experimental/CommonUI/Source/CommonUI/Private/CommonListView.cpp
@@ -3,7 +3,7 @@
 #include "CommonListView.h"
 #include "CommonWidgetPaletteCategories.h"
 #include "SCommonButtonTableRow.h"
-#include "CommonUIPrivatePCH.h"
+#include "CommonUIPrivate.h"
 
 //////////////////////////////////////////////////////////////////////////
 // UCommonListView
diff --git a/Engine/Plugins/Experimental/CommonUI/Source/CommonUI/Private/CommonLoadGuard.cpp b/Engine/Plugins/Experimental/CommonUI/Source/CommonUI/Private/CommonLoadGuard.cpp
--- a/Engine/Plugins/Experimental/CommonUI/Source/CommonUI/Private/CommonLoadGuard.cpp
+++ b/Engine/Plugins/Experimental/CommonUI/Source/CommonUI/Private/CommonLoadGuard.cpp
@@ -1,7 +1,7 @@
 // Copyright Epic Games, Inc. All Rights Reserved.
 
 #include "CommonLoadGuard.h"
-#include "CommonUIPrivatePCH.h"
+#include "CommonUIPrivate.h"
 #include "Components/SizeBoxSlot.h"
 #include "CommonTextBlock.h"
 #include "CommonWidgetPaletteCategories.h"
diff --git a/Engine/Plugins/Experimental/CommonUI/Source/CommonUI/Private/CommonNumericTextBlock.cpp b/Engine/Plugins/Experimental/CommonUI/Source/CommonUI/Private/CommonNumericTextBlock.cpp
--- a/Engine/Plugins/Experimental/CommonUI/Source/CommonUI/Private/CommonNumericTextBlock.cpp
+++ b/Engine/Plugins/Experimental/CommonUI/Source/CommonUI/Private/CommonNumericTextBlock.cpp
@@ -1,7 +1,7 @@
 // Copyright Epic Games, Inc. All Rights Reserved.
 
 #include "CommonNumericTextBlock.h"
-#include "CommonUIPrivatePCH.h"
+#include "CommonUIPrivate.h"
 #include "CommonWidgetPaletteCategories.h"
 #include "TimerManager.h"
 #include "Widgets/Text/STextBlock.h"
diff --git a/Engine/Plugins/Experimental/CommonUI/Source/CommonUI/Private/CommonRotator.cpp b/Engine/Plugins/Experimental/CommonUI/Source/CommonUI/Private/CommonRotator.cpp
--- a/Engine/Plugins/Experimental/CommonUI/Source/CommonUI/Private/CommonRotator.cpp
+++ b/Engine/Plugins/Experimental/CommonUI/Source/CommonUI/Private/CommonRotator.cpp
@@ -1,7 +1,7 @@
 // Copyright Epic Games, Inc. All Rights Reserved.
 
 #include "CommonRotator.h"
-#include "CommonUIPrivatePCH.h"
+#include "CommonUIPrivate.h"
 
 #define LOCTEXT_NAMESPACE "UMG"
 
diff --git a/Engine/Plugins/Experimental/CommonUI/Source/CommonUI/Private/CommonTabListWidgetBase.cpp b/Engine/Plugins/Experimental/CommonUI/Source/CommonUI/Private/CommonTabListWidgetBase.cpp
--- a/Engine/Plugins/Experimental/CommonUI/Source/CommonUI/Private/CommonTabListWidgetBase.cpp
+++ b/Engine/Plugins/Experimental/CommonUI/Source/CommonUI/Private/CommonTabListWidgetBase.cpp
@@ -1,7 +1,7 @@
 // Copyright Epic Games, Inc. All Rights Reserved.
 
 #include "CommonTabListWidgetBase.h"
-#include "CommonUIPrivatePCH.h"
+#include "CommonUIPrivate.h"
 
 #include "Blueprint/WidgetLayoutLibrary.h"
 #include "CommonAnimatedSwitcher.h"
diff --git a/Engine/Plugins/Experimental/CommonUI/Source/CommonUI/Private/CommonTextBlock.cpp b/Engine/Plugins/Experimental/CommonUI/Source/CommonUI/Private/CommonTextBlock.cpp
--- a/Engine/Plugins/Experimental/CommonUI/Source/CommonUI/Private/CommonTextBlock.cpp
+++ b/Engine/Plugins/Experimental/CommonUI/Source/CommonUI/Private/CommonTextBlock.cpp
@@ -1,7 +1,7 @@
 // Copyright Epic Games, Inc. All Rights Reserved.
 
 #include "CommonTextBlock.h"
-#include "CommonUIPrivatePCH.h"
+#include "CommonUIPrivate.h"
 #include "Containers/Ticker.h"
 #include "CommonUIEditorSettings.h"
 #include "CommonUIUtils.h"
diff --git a/Engine/Plugins/Experimental/CommonUI/Source/CommonUI/Private/CommonTileView.cpp b/Engine/Plugins/Experimental/CommonUI/Source/CommonUI/Private/CommonTileView.cpp
--- a/Engine/Plugins/Experimental/CommonUI/Source/CommonUI/Private/CommonTileView.cpp
+++ b/Engine/Plugins/Experimental/CommonUI/Source/CommonUI/Private/CommonTileView.cpp
@@ -1,7 +1,7 @@
 // Copyright Epic Games, Inc. All Rights Reserved.
 
 #include "CommonTileView.h"
-#include "CommonUIPrivatePCH.h"
+#include "CommonUIPrivate.h"
 #include "SCommonButtonTableRow.h"
 #include "SCommonTileView.h"
 
diff --git a/Engine/Plugins/Experimental/CommonUI/Source/CommonUI/Private/CommonTreeView.cpp b/Engine/Plugins/Experimental/CommonUI/Source/CommonUI/Private/CommonTreeView.cpp
--- a/Engine/Plugins/Experimental/CommonUI/Source/CommonUI/Private/CommonTreeView.cpp
+++ b/Engine/Plugins/Experimental/CommonUI/Source/CommonUI/Private/CommonTreeView.cpp
@@ -1,7 +1,7 @@
 // Copyright Epic Games, Inc. All Rights Reserved.
 
 #include "CommonTreeView.h"
-#include "CommonUIPrivatePCH.h"
+#include "CommonUIPrivate.h"
 #include "SCommonButtonTableRow.h"
 
 UCommonTreeView::UCommonTreeView(const FObjectInitializer& ObjectInitializer)
diff --git a/Engine/Plugins/Experimental/CommonUI/Source/CommonUI/Private/CommonUIModule.cpp b/Engine/Plugins/Experimental/CommonUI/Source/CommonUI/Private/CommonUIModule.cpp
--- a/Engine/Plugins/Experimental/CommonUI/Source/CommonUI/Private/CommonUIModule.cpp
+++ b/Engine/Plugins/Experimental/CommonUI/Source/CommonUI/Private/CommonUIModule.cpp
@@ -1,6 +1,6 @@
 // Copyright Epic Games, Inc. All Rights Reserved.
 
-#include "CommonUIPrivatePCH.h"
+#include "CommonUIPrivate.h"
 #include "ICommonUIModule.h"
 #include "CommonUISettings.h"
 #include "CommonUIEditorSettings.h"
diff --git a/Engine/Plugins/Experimental/CommonUI/Source/CommonUI/Private/CommonUIPrivate.h b/Engine/Plugins/Experimental/CommonUI/Source/CommonUI/Private/CommonUIPrivate.h
new file mode 100644
--- /dev/null
+++ b/Engine/Plugins/Experimental/CommonUI/Source/CommonUI/Private/CommonUIPrivate.h
@@ -0,0 +1,11 @@
+// Copyright Epic Games, Inc. All Rights Reserved.
+
+#pragma once
+
+
+/* Dependencies
+ *****************************************************************************/
+
+#include "Styling/SlateStyle.h"
+
+DECLARE_LOG_CATEGORY_EXTERN(LogCommonUI, Log, All);
diff --git a/Engine/Plugins/Experimental/CommonUI/Source/CommonUI/Private/CommonUIPrivatePCH.h b/Engine/Plugins/Experimental/CommonUI/Source/CommonUI/Private/CommonUIPrivatePCH.h
deleted file mode 100644
diff --git a/Engine/Plugins/Experimental/CommonUI/Source/CommonUI/Private/CommonUISubsystemBase.cpp b/Engine/Plugins/Experimental/CommonUI/Source/CommonUI/Private/CommonUISubsystemBase.cpp
--- a/Engine/Plugins/Experimental/CommonUI/Source/CommonUI/Private/CommonUISubsystemBase.cpp
+++ b/Engine/Plugins/Experimental/CommonUI/Source/CommonUI/Private/CommonUISubsystemBase.cpp
@@ -5,7 +5,7 @@
 #include "Subsystems/SubsystemCollection.h"
 
 #include "Blueprint/UserWidget.h"
-#include "CommonUIPrivatePCH.h"
+#include "CommonUIPrivate.h"
 #include "CommonInputSubsystem.h"
 #include "ICommonInputModule.h"
 #include "CommonInputSettings.h"
diff --git a/Engine/Plugins/Experimental/CommonUI/Source/CommonUI/Private/CommonUITypes.cpp b/Engine/Plugins/Experimental/CommonUI/Source/CommonUI/Private/CommonUITypes.cpp
--- a/Engine/Plugins/Experimental/CommonUI/Source/CommonUI/Private/CommonUITypes.cpp
+++ b/Engine/Plugins/Experimental/CommonUI/Source/CommonUI/Private/CommonUITypes.cpp
@@ -1,7 +1,7 @@
 // Copyright Epic Games, Inc. All Rights Reserved.
  
 #include "CommonUITypes.h"
-#include "CommonUIPrivatePCH.h"
+#include "CommonUIPrivate.h"
 #include "ICommonInputModule.h"
 #include "CommonInputSettings.h"
 #include "HAL/PlatformInput.h"
diff --git a/Engine/Plugins/Experimental/CommonUI/Source/CommonUI/Private/CommonUIUtils.cpp b/Engine/Plugins/Experimental/CommonUI/Source/CommonUI/Private/CommonUIUtils.cpp
--- a/Engine/Plugins/Experimental/CommonUI/Source/CommonUI/Private/CommonUIUtils.cpp
+++ b/Engine/Plugins/Experimental/CommonUI/Source/CommonUI/Private/CommonUIUtils.cpp
@@ -2,7 +2,7 @@
  
 #include "CommonUIUtils.h"
 #include "ICommonUIModule.h"
-#include "CommonUIPrivatePCH.h"
+#include "CommonUIPrivate.h"
 #include "CommonUIEditorSettings.h"
 #include "Blueprint/WidgetTree.h"
 #include "Editor/WidgetCompilerLog.h"
diff --git a/Engine/Plugins/Experimental/CommonUI/Source/CommonUI/Private/CommonUserWidget.cpp b/Engine/Plugins/Experimental/CommonUI/Source/CommonUI/Private/CommonUserWidget.cpp
--- a/Engine/Plugins/Experimental/CommonUI/Source/CommonUI/Private/CommonUserWidget.cpp
+++ b/Engine/Plugins/Experimental/CommonUI/Source/CommonUI/Private/CommonUserWidget.cpp
@@ -1,7 +1,7 @@
 // Copyright Epic Games, Inc. All Rights Reserved.
 
 #include "CommonUserWidget.h"
-#include "CommonUIPrivatePCH.h"
+#include "CommonUIPrivate.h"
 
 #include "Engine/GameInstance.h"
 #include "CommonInputSubsystem.h"
diff --git a/Engine/Plugins/Experimental/CommonUI/Source/CommonUI/Private/CommonVisibilitySwitcher.cpp b/Engine/Plugins/Experimental/CommonUI/Source/CommonUI/Private/CommonVisibilitySwitcher.cpp
--- a/Engine/Plugins/Experimental/CommonUI/Source/CommonUI/Private/CommonVisibilitySwitcher.cpp
+++ b/Engine/Plugins/Experimental/CommonUI/Source/CommonUI/Private/CommonVisibilitySwitcher.cpp
@@ -3,7 +3,7 @@
 #include "CommonVisibilitySwitcher.h"
 
 #include "CommonActivatableWidget.h"
-#include "CommonUIPrivatePCH.h"
+#include "CommonUIPrivate.h"
 #include "CommonVisibilitySwitcherSlot.h"
 #include "CommonWidgetPaletteCategories.h"
 #include "Widgets/Layout/SBox.h"
diff --git a/Engine/Plugins/Experimental/CommonUI/Source/CommonUI/Private/CommonVisibilityWidgetBase.cpp b/Engine/Plugins/Experimental/CommonUI/Source/CommonUI/Private/CommonVisibilityWidgetBase.cpp
--- a/Engine/Plugins/Experimental/CommonUI/Source/CommonUI/Private/CommonVisibilityWidgetBase.cpp
+++ b/Engine/Plugins/Experimental/CommonUI/Source/CommonUI/Private/CommonVisibilityWidgetBase.cpp
@@ -1,7 +1,7 @@
 // Copyright Epic Games, Inc. All Rights Reserved.
 
 #include "CommonVisibilityWidgetBase.h"
-#include "CommonUIPrivatePCH.h"
+#include "CommonUIPrivate.h"
 #include "CommonUISubsystemBase.h"
 
 #include "Components/BorderSlot.h"
diff --git a/Engine/Plugins/Experimental/CommonUI/Source/CommonUI/Private/CommonWidgetCarousel.cpp b/Engine/Plugins/Experimental/CommonUI/Source/CommonUI/Private/CommonWidgetCarousel.cpp
--- a/Engine/Plugins/Experimental/CommonUI/Source/CommonUI/Private/CommonWidgetCarousel.cpp
+++ b/Engine/Plugins/Experimental/CommonUI/Source/CommonUI/Private/CommonWidgetCarousel.cpp
@@ -1,7 +1,7 @@
 // Copyright Epic Games, Inc. All Rights Reserved.
 
 #include "CommonWidgetCarousel.h"
-#include "CommonUIPrivatePCH.h"
+#include "CommonUIPrivate.h"
 #include "CommonWidgetPaletteCategories.h"
 #include "Containers/Ticker.h"
 
diff --git a/Engine/Plugins/Experimental/CommonUI/Source/CommonUI/Private/CommonWidgetCarouselNavBar.cpp b/Engine/Plugins/Experimental/CommonUI/Source/CommonUI/Private/CommonWidgetCarouselNavBar.cpp
--- a/Engine/Plugins/Experimental/CommonUI/Source/CommonUI/Private/CommonWidgetCarouselNavBar.cpp
+++ b/Engine/Plugins/Experimental/CommonUI/Source/CommonUI/Private/CommonWidgetCarouselNavBar.cpp
@@ -1,7 +1,7 @@
 // Copyright Epic Games, Inc. All Rights Reserved.
 
 #include "CommonWidgetCarouselNavBar.h"
-#include "CommonUIPrivatePCH.h"
+#include "CommonUIPrivate.h"
 #include "CommonWidgetPaletteCategories.h"
 #include "CommonWidgetCarousel.h"
 #include "Groups/CommonButtonGroupBase.h"
diff --git a/Engine/Plugins/Experimental/CommonUI/Source/CommonUI/Private/Groups/CommonButtonGroupBase.cpp b/Engine/Plugins/Experimental/CommonUI/Source/CommonUI/Private/Groups/CommonButtonGroupBase.cpp
--- a/Engine/Plugins/Experimental/CommonUI/Source/CommonUI/Private/Groups/CommonButtonGroupBase.cpp
+++ b/Engine/Plugins/Experimental/CommonUI/Source/CommonUI/Private/Groups/CommonButtonGroupBase.cpp
@@ -1,7 +1,7 @@
 // Copyright Epic Games, Inc. All Rights Reserved.
 
 #include "Groups/CommonButtonGroupBase.h"
-#include "CommonUIPrivatePCH.h"
+#include "CommonUIPrivate.h"
 #include <functional>
 
 UCommonButtonGroupBase::UCommonButtonGroupBase()
diff --git a/Engine/Plugins/Experimental/CommonUI/Source/CommonUI/Private/Groups/CommonWidgetGroupBase.cpp b/Engine/Plugins/Experimental/CommonUI/Source/CommonUI/Private/Groups/CommonWidgetGroupBase.cpp
--- a/Engine/Plugins/Experimental/CommonUI/Source/CommonUI/Private/Groups/CommonWidgetGroupBase.cpp
+++ b/Engine/Plugins/Experimental/CommonUI/Source/CommonUI/Private/Groups/CommonWidgetGroupBase.cpp
@@ -1,7 +1,7 @@
 // Copyright Epic Games, Inc. All Rights Reserved.
 
 #include "Groups/CommonWidgetGroupBase.h"
-#include "CommonUIPrivatePCH.h"
+#include "CommonUIPrivate.h"
 
 UCommonWidgetGroupBase::UCommonWidgetGroupBase()
 {
diff --git a/Engine/Plugins/Experimental/CommonUI/Source/CommonUI/Private/Input/CommonAnalogCursor.cpp b/Engine/Plugins/Experimental/CommonUI/Source/CommonUI/Private/Input/CommonAnalogCursor.cpp
--- a/Engine/Plugins/Experimental/CommonUI/Source/CommonUI/Private/Input/CommonAnalogCursor.cpp
+++ b/Engine/Plugins/Experimental/CommonUI/Source/CommonUI/Private/Input/CommonAnalogCursor.cpp
@@ -1,7 +1,7 @@
 // Copyright Epic Games, Inc. All Rights Reserved.
 
 #include "Input/CommonAnalogCursor.h"
-#include "CommonUIPrivatePCH.h"
+#include "CommonUIPrivate.h"
 #include "Slate/SObjectWidget.h"
 #include "Blueprint/UserWidget.h"
 #include "Input/CommonUIActionRouterBase.h"
diff --git a/Engine/Plugins/Experimental/CommonUI/Source/CommonUI/Private/Input/CommonUIActionRouterBase.cpp b/Engine/Plugins/Experimental/CommonUI/Source/CommonUI/Private/Input/CommonUIActionRouterBase.cpp
--- a/Engine/Plugins/Experimental/CommonUI/Source/CommonUI/Private/Input/CommonUIActionRouterBase.cpp
+++ b/Engine/Plugins/Experimental/CommonUI/Source/CommonUI/Private/Input/CommonUIActionRouterBase.cpp
@@ -18,7 +18,7 @@
 #include "Stats/Stats.h"
 
 #include "CommonGameViewportClient.h"
-#include "CommonUIPrivatePCH.h"
+#include "CommonUIPrivate.h"
 #include "CommonActivatableWidget.h"
 #include "CommonUIUtils.h"
 #include "CommonUISubsystemBase.h"
diff --git a/Engine/Plugins/Experimental/CommonUI/Source/CommonUI/Private/Input/UIActionRouterTypes.cpp b/Engine/Plugins/Experimental/CommonUI/Source/CommonUI/Private/Input/UIActionRouterTypes.cpp
--- a/Engine/Plugins/Experimental/CommonUI/Source/CommonUI/Private/Input/UIActionRouterTypes.cpp
+++ b/Engine/Plugins/Experimental/CommonUI/Source/CommonUI/Private/Input/UIActionRouterTypes.cpp
@@ -4,7 +4,7 @@
 #include "Input/CommonUIActionRouterBase.h"
 #include "CommonActivatableWidget.h"
 #include "Input/CommonUIInputTypes.h"
-#include "CommonUIPrivatePCH.h"
+#include "CommonUIPrivate.h"
 #include "CommonInputSettings.h"
 #include "CommonInputSubsystem.h"
 #include "ICommonInputModule.h"
diff --git a/Engine/Plugins/Experimental/CommonUI/Source/CommonUIEditor/Private/CommonUIEditorModule.cpp b/Engine/Plugins/Experimental/CommonUI/Source/CommonUIEditor/Private/CommonUIEditorModule.cpp
--- a/Engine/Plugins/Experimental/CommonUI/Source/CommonUIEditor/Private/CommonUIEditorModule.cpp
+++ b/Engine/Plugins/Experimental/CommonUI/Source/CommonUIEditor/Private/CommonUIEditorModule.cpp
@@ -1,7 +1,7 @@
 // Copyright Epic Games, Inc. All Rights Reserved.
 
 #include "CommonUIEditorModule.h"
-#include "CommonUIPrivatePCH.h"
+#include "CommonUIPrivate.h"
 #include "CommonVideoPlayerCustomization.h"
 
 #include "GameplayTagsEditorModule.h"
