Exposes bIsPersistentBinding and InputModeOverride using protected getters so subclasses have access to them.

diff --git a/Engine/Plugins/Experimental/CommonUI/Source/CommonUI/Public/CommonButtonBase.h b/Engine/Plugins/Experimental/CommonUI/Source/CommonUI/Public/CommonButtonBase.h
--- a/Engine/Plugins/Experimental/CommonUI/Source/CommonUI/Public/CommonButtonBase.h
+++ b/Engine/Plugins/Experimental/CommonUI/Source/CommonUI/Public/CommonButtonBase.h
@@ -454,6 +454,9 @@ public:
 	FOnIsSelectedChanged& OnIsSelectedChanged() const { return OnIsSelectedChangedEvent; }
 
 protected:
+	bool IsPersistentBinding() const { return bIsPersistentBinding; }
+	ECommonInputMode GetInputModeOverride() const { return InputModeOverride; }
+	
 	virtual UCommonButtonInternalBase* ConstructInternalButton();
 
 	virtual void OnWidgetRebuilt();
