#rnx

diff --git a/Engine/Plugins/Experimental/CommonUI/Source/CommonInput/Public/CommonInputActionDomain.h b/Engine/Plugins/Experimental/CommonUI/Source/CommonInput/Public/CommonInputActionDomain.h
--- a/Engine/Plugins/Experimental/CommonUI/Source/CommonInput/Public/CommonInputActionDomain.h
+++ b/Engine/Plugins/Experimental/CommonUI/Source/CommonInput/Public/CommonInputActionDomain.h
@@ -4,6 +4,7 @@
 
 #include "Engine/DataAsset.h"
 #include "Types/ISlateMetaData.h"
+#include "CommonInputModeTypes.h"
 
 #include "CommonInputActionDomain.generated.h"
 
diff --git a/Engine/Plugins/Experimental/CommonUI/Source/CommonInput/Public/CommonInputModeTypes.h b/Engine/Plugins/Experimental/CommonUI/Source/CommonInput/Public/CommonInputModeTypes.h
new file mode 100644
--- /dev/null
+++ b/Engine/Plugins/Experimental/CommonUI/Source/CommonInput/Public/CommonInputModeTypes.h
@@ -0,0 +1,24 @@
+// Copyright Epic Games, Inc. All Rights Reserved.
+
+#pragma once
+
+UENUM(BlueprintType)
+enum class ECommonInputMode : uint8
+{
+	Menu	UMETA(Tooltip = "Input is received by the UI only"),
+	Game	UMETA(Tooltip = "Input is received by the Game only"),
+	All		UMETA(Tooltip = "Input is received by UI and the Game"),
+
+	MAX UMETA(Hidden)
+};
+
+inline const TCHAR* LexToString(ECommonInputMode Value)
+{
+	switch (Value)
+	{
+	case ECommonInputMode::Menu: return TEXT("Menu"); break;
+	case ECommonInputMode::Game: return TEXT("Game"); break;
+	case ECommonInputMode::All: return TEXT("All"); break;
+	}
+	return TEXT("Unknown");
+}
diff --git a/Engine/Plugins/Experimental/CommonUI/Source/CommonUI/Public/CommonButtonBase.h b/Engine/Plugins/Experimental/CommonUI/Source/CommonUI/Public/CommonButtonBase.h
--- a/Engine/Plugins/Experimental/CommonUI/Source/CommonUI/Public/CommonButtonBase.h
+++ b/Engine/Plugins/Experimental/CommonUI/Source/CommonUI/Public/CommonButtonBase.h
@@ -7,7 +7,7 @@
 #include "Components/Button.h"
 #include "Misc/Optional.h"
 #include "Types/ISlateMetaData.h"
-#include "Input/CommonInputMode.h"
+#include "CommonInputModeTypes.h"
 #include "CommonButtonBase.generated.h"
 
 class UCommonTextStyle;
diff --git a/Engine/Plugins/Experimental/CommonUI/Source/CommonUI/Public/Input/CommonInputMode.h b/Engine/Plugins/Experimental/CommonUI/Source/CommonUI/Public/Input/CommonInputMode.h
--- a/Engine/Plugins/Experimental/CommonUI/Source/CommonUI/Public/Input/CommonInputMode.h
+++ b/Engine/Plugins/Experimental/CommonUI/Source/CommonUI/Public/Input/CommonInputMode.h
@@ -2,16 +2,10 @@
 
 #pragma once
 
-#include "Misc/EnumRange.h"
+#ifdef _MSC_VER
+#pragma message(__FILE__"(8): warning: use CommonInputModeTypes.h from CommonInput module instead of Input/CommonInputMode.h from CommonUI module")
+#else
+#pragma message("#include CommonInputModeTypes.h from CommonInput module instead of Input/CommonInputMode.h from CommonUI module")
+#endif
 
-#include "CommonInputMode.generated.h"
-
-UENUM(BlueprintType)
-enum class ECommonInputMode : uint8
-{
-	Menu	UMETA(Tooltip="Input is received by the UI only"),
-	Game	UMETA(Tooltip="Input is received by the Game only"),
-	All		UMETA(Tooltip="Input is received by UI and the Game"),
-	
-	MAX UMETA(Hidden)
-};
+#include "CommonInputModeTypes.h"
\ No newline at end of file
diff --git a/Engine/Plugins/Experimental/CommonUI/Source/CommonUI/Public/Input/CommonUIInputTypes.h b/Engine/Plugins/Experimental/CommonUI/Source/CommonUI/Public/Input/CommonUIInputTypes.h
--- a/Engine/Plugins/Experimental/CommonUI/Source/CommonUI/Public/Input/CommonUIInputTypes.h
+++ b/Engine/Plugins/Experimental/CommonUI/Source/CommonUI/Public/Input/CommonUIInputTypes.h
@@ -5,20 +5,9 @@
 #include "UIActionBindingHandle.h"
 #include "UITag.h"
 #include "Misc/EnumRange.h"
-#include "CommonInputMode.h"
+#include "CommonInputModeTypes.h"
 #include "InputCoreTypes.h"
 
-inline const TCHAR* LexToString(ECommonInputMode Value)
-{
-	switch (Value)
-	{
-	case ECommonInputMode::Menu: return TEXT("Menu"); break;
-	case ECommonInputMode::Game: return TEXT("Game"); break;
-	case ECommonInputMode::All: return TEXT("All"); break;
-	}
-	return TEXT("Unknown");
-}
-
 struct COMMONUI_API FBindUIActionArgs
 {
 	FBindUIActionArgs(FUIActionTag InActionTag, const FSimpleDelegate& InOnExecuteAction)
