- Add sound overrides for the CommonButtonBase Normal, Selected and Locked state

diff --git a/Engine/Plugins/Experimental/CommonUI/Source/CommonUI/Private/CommonButtonBase.cpp b/Engine/Plugins/Experimental/CommonUI/Source/CommonUI/Private/CommonButtonBase.cpp
--- a/Engine/Plugins/Experimental/CommonUI/Source/CommonUI/Private/CommonButtonBase.cpp
+++ b/Engine/Plugins/Experimental/CommonUI/Source/CommonUI/Private/CommonButtonBase.cpp
@@ -734,6 +734,8 @@ void UCommonButtonBase::SetIsLocked(bool bInIsLocked)
 {
 	bLocked = bInIsLocked;
 
+	SetButtonStyle();
+
 	BP_OnLockedChanged(bLocked);
 }
 
@@ -1110,6 +1112,8 @@ FReply UCommonButtonBase::NativeOnFocusReceived(const FGeometry& InGeometry, con
 void UCommonButtonBase::NativeOnFocusLost(const FFocusEvent& InFocusEvent)
 {
 	Super::NativeOnFocusLost(InFocusEvent);
+
+	HandleFocusLost();
 }
 
 void UCommonButtonBase::NativeOnSelected(bool bBroadcast)
@@ -1140,6 +1144,12 @@ void UCommonButtonBase::NativeOnHovered()
 {
 	BP_OnHovered();
 	OnHovered().Broadcast();
+
+	if (bShouldFocusOnHover && bIsFocusable && RootButton->IsFocusable)
+	{
+		RootButton->SetFocus();
+	}
+	
 	if (OnButtonBaseHovered.IsBound())
 	{
 		OnButtonBaseHovered.Broadcast(this);
@@ -1155,6 +1165,7 @@ void UCommonButtonBase::NativeOnUnhovered()
 {
 	BP_OnUnhovered();
 	OnUnhovered().Broadcast();
+	
 	if (OnButtonBaseUnhovered.IsBound())
 	{
 		OnButtonBaseUnhovered.Broadcast(this);
@@ -1168,7 +1179,7 @@ void UCommonButtonBase::NativeOnUnhovered()
 
 void UCommonButtonBase::NativeOnClicked()
 {
-	if(!GetLocked())
+	if (!GetLocked())
 	{
 		BP_OnClicked();
 		OnClicked().Broadcast();
@@ -1198,7 +1209,7 @@ void UCommonButtonBase::NativeOnClicked()
 
 void UCommonButtonBase::NativeOnDoubleClicked()
 {
-	if(!GetLocked())
+	if (!GetLocked())
 	{
 		BP_OnDoubleClicked();
 		OnDoubleClicked().Broadcast();
@@ -1290,6 +1301,8 @@ void UCommonButtonBase::BuildStyles()
 		NormalStyle.Disabled = CommonButtonStyle->bSingleMaterial ? DynamicSingleMaterialBrush : DisabledBrush;
 		NormalStyle.NormalPadding = ButtonPadding;
 		NormalStyle.PressedPadding = ButtonPadding;
+
+		// Sets the sound overrides for the Normal state
 		NormalStyle.PressedSlateSound = bHasPressedSlateSoundOverride ? PressedSlateSoundOverride : CommonButtonStyle->PressedSlateSound;
 		NormalStyle.HoveredSlateSound = bHasHoveredSlateSoundOverride ? HoveredSlateSoundOverride : CommonButtonStyle->HoveredSlateSound;
 
@@ -1299,24 +1312,48 @@ void UCommonButtonBase::BuildStyles()
 		SelectedStyle.Disabled = CommonButtonStyle->bSingleMaterial ? DynamicSingleMaterialBrush : DisabledBrush;
 		SelectedStyle.NormalPadding = ButtonPadding;
 		SelectedStyle.PressedPadding = ButtonPadding;
-		SelectedStyle.PressedSlateSound =
+
+		
+		 // Selected State Sound overrides
+		if (SelectedPressedSlateSoundOverride.GetResourceObject())
+		{
+			SelectedStyle.PressedSlateSound = SelectedPressedSlateSoundOverride;
+		}
+		else
+		{
+			SelectedStyle.PressedSlateSound =
 			bHasPressedSlateSoundOverride || !CommonButtonStyle->SelectedPressedSlateSound
 			? NormalStyle.PressedSlateSound
 			: CommonButtonStyle->SelectedPressedSlateSound.Sound;
-
-		SelectedStyle.HoveredSlateSound =
+		}
+		
+		if (SelectedHoveredSlateSoundOverride.GetResourceObject())
+		{
+			SelectedStyle.HoveredSlateSound = SelectedHoveredSlateSoundOverride;
+		}
+		else
+		{
+			SelectedStyle.HoveredSlateSound =
 			bHasHoveredSlateSoundOverride || !CommonButtonStyle->SelectedHoveredSlateSound
 			? NormalStyle.HoveredSlateSound
 			: CommonButtonStyle->SelectedHoveredSlateSound.Sound;
+		}
 
-		DisabledStyle = NormalStyle;
-		if (!bHasPressedSlateSoundOverride && CommonButtonStyle->DisabledPressedSlateSound)
+		// Locked State Sound overrides
+		LockedStyle = NormalStyle;
+		if (CommonButtonStyle->LockedPressedSlateSound || LockedPressedSlateSoundOverride.GetResourceObject())
 		{
-			DisabledStyle.PressedSlateSound = CommonButtonStyle->DisabledPressedSlateSound.Sound;
+			LockedStyle.PressedSlateSound =
+			LockedPressedSlateSoundOverride.GetResourceObject()
+			? LockedPressedSlateSoundOverride
+			: CommonButtonStyle->LockedPressedSlateSound.Sound;
 		}
-		if (!bHasHoveredSlateSoundOverride && CommonButtonStyle->DisabledHoveredSlateSound)
+		if (CommonButtonStyle->LockedHoveredSlateSound || LockedHoveredSlateSoundOverride.GetResourceObject())
 		{
-			DisabledStyle.HoveredSlateSound = CommonButtonStyle->DisabledHoveredSlateSound.Sound;
+			LockedStyle.HoveredSlateSound =
+			LockedHoveredSlateSoundOverride.GetResourceObject()
+			? LockedHoveredSlateSoundOverride
+			: CommonButtonStyle->LockedHoveredSlateSound.Sound;
 		}
 
 		SetButtonStyle();
@@ -1330,7 +1367,11 @@ void UCommonButtonBase::SetButtonStyle()
 	if (UButton* ButtonPtr = RootButton.Get())
 	{
 		const FButtonStyle* UseStyle;
-		if (bSelected)
+		if (bLocked)
+		{
+			UseStyle = &LockedStyle;
+		}
+		else if (bSelected)
 		{
 			UseStyle = &SelectedStyle;
 		}
@@ -1373,6 +1414,42 @@ void UCommonButtonBase::SetHoveredSoundOverride(USoundBase* Sound)
 	}
 }
 
+void UCommonButtonBase::SetSelectedPressedSoundOverride(USoundBase* Sound)
+{
+	if (SelectedPressedSlateSoundOverride.GetResourceObject() != Sound)
+	{
+		SelectedPressedSlateSoundOverride.SetResourceObject(Sound);
+		BuildStyles();
+	}
+}
+
+void UCommonButtonBase::SetSelectedHoveredSoundOverride(USoundBase* Sound)
+{
+	if (SelectedHoveredSlateSoundOverride.GetResourceObject() != Sound)
+	{
+		SelectedHoveredSlateSoundOverride.SetResourceObject(Sound);
+		BuildStyles();
+	}
+}
+
+void UCommonButtonBase::SetLockedPressedSoundOverride(USoundBase* Sound)
+{
+	if (LockedPressedSlateSoundOverride.GetResourceObject() != Sound)
+	{
+		LockedPressedSlateSoundOverride.SetResourceObject(Sound);
+		BuildStyles();
+	}
+}
+
+void UCommonButtonBase::SetLockedHoveredSoundOverride(USoundBase* Sound)
+{
+	if (LockedHoveredSlateSoundOverride.GetResourceObject() != Sound)
+	{
+		LockedHoveredSlateSoundOverride.SetResourceObject(Sound);
+		BuildStyles();
+	}
+}
+
 void UCommonButtonBase::UpdateInputActionWidgetVisibility()
 {
 	if (InputActionWidget)
diff --git a/Engine/Plugins/Experimental/CommonUI/Source/CommonUI/Public/CommonButtonBase.h b/Engine/Plugins/Experimental/CommonUI/Source/CommonUI/Public/CommonButtonBase.h
--- a/Engine/Plugins/Experimental/CommonUI/Source/CommonUI/Public/CommonButtonBase.h
+++ b/Engine/Plugins/Experimental/CommonUI/Source/CommonUI/Public/CommonButtonBase.h
@@ -136,11 +136,10 @@ public:
 	UPROPERTY(EditDefaultsOnly, BlueprintReadOnly, Category = "Properties", meta = (DisplayName = "Selected Pressed Sound"))
 	FCommonButtonStyleOptionalSlateSound SelectedPressedSlateSound;
 
-	/** The sound to play when the button is pressed while disabled */
-	UPROPERTY(EditDefaultsOnly, BlueprintReadOnly, Category = "Properties", meta = (DisplayName = "Disabled Pressed Sound"))
-	FCommonButtonStyleOptionalSlateSound DisabledPressedSlateSound;
-
-
+	/** The sound to play when the button is pressed while locked */
+	UPROPERTY(EditDefaultsOnly, BlueprintReadOnly, Category = "Properties", meta = (DisplayName = "Locked Pressed Sound"))
+	FCommonButtonStyleOptionalSlateSound LockedPressedSlateSound;
+	
 	/** The sound to play when the button is hovered */
 	UPROPERTY(EditDefaultsOnly, BlueprintReadOnly, Category = "Properties", meta = (DisplayName = "Hovered Sound"))
 	FSlateSound HoveredSlateSound;
@@ -148,10 +147,10 @@ public:
 	/** The sound to play when the button is hovered while selected */
 	UPROPERTY(EditDefaultsOnly, BlueprintReadOnly, Category = "Properties", meta = (DisplayName = "Selected Hovered Sound"))
 	FCommonButtonStyleOptionalSlateSound SelectedHoveredSlateSound;
-
-	/** The sound to play when the button is hovered while disabled */
-	UPROPERTY(EditDefaultsOnly, BlueprintReadOnly, Category = "Properties", meta = (DisplayName = "Disabled Hovered Sound"))
-	FCommonButtonStyleOptionalSlateSound DisabledHoveredSlateSound;
+	
+	/** The sound to play when the button is hovered while locked */
+	UPROPERTY(EditDefaultsOnly, BlueprintReadOnly, Category = "Properties", meta = (DisplayName = "Locked Hovered Sound"))
+	FCommonButtonStyleOptionalSlateSound LockedHoveredSlateSound;
 
 	UFUNCTION(BlueprintCallable, Category = "Common ButtonStyle|Getters")
 	void GetMaterialBrush(FSlateBrush& Brush) const;
@@ -438,6 +437,18 @@ public:
 	UFUNCTION(BlueprintCallable, Category = "Common Button|Sound")
 	void SetHoveredSoundOverride(USoundBase* Sound);
 
+	UFUNCTION(BlueprintCallable, Category = "Common Button|Sound")
+	void SetSelectedPressedSoundOverride(USoundBase* Sound);
+
+	UFUNCTION(BlueprintCallable, Category = "Common Button|Sound")
+	void SetSelectedHoveredSoundOverride(USoundBase* Sound);
+	
+	UFUNCTION(BlueprintCallable, Category = "Common Button|Sound")
+	void SetLockedPressedSoundOverride(USoundBase* Sound);
+
+	UFUNCTION(BlueprintCallable, Category = "Common Button|Sound")
+	void SetLockedHoveredSoundOverride(USoundBase* Sound);
+
 	DECLARE_EVENT(UCommonButtonBase, FCommonButtonEvent);
 	FCommonButtonEvent& OnClicked() const { return OnClickedEvent; }
 	FCommonButtonEvent& OnDoubleClicked() const { return OnDoubleClickedEvent; }
@@ -629,21 +640,55 @@ protected:
 	UPROPERTY(EditAnywhere, BlueprintReadOnly, Category = Style)
 	bool bHideInputAction;
 
-	/** Optional override for the sound to play when this button is pressed */
+	/**
+	 * Optional override for the sound to play when this button is pressed.
+	 * Also used for the Selected and Locked Pressed state if their respective overrides are empty.
+	 */
 	UPROPERTY(EditAnywhere, BlueprintReadOnly, Category = Sound, meta = (DisplayName = "Pressed Sound Override"))
 	FSlateSound PressedSlateSoundOverride;
 
-	/** Optional override for the sound to play when this button is hovered */
+	/**
+	 * Optional override for the sound to play when this button is hovered.
+	 * Also used for the Selected and Locked Hovered state if their respective overrides are empty.
+	 */
 	UPROPERTY(EditAnywhere, BlueprintReadOnly, Category = Sound, meta = (DisplayName = "Hovered Sound Override"))
 	FSlateSound HoveredSlateSoundOverride;
 
+	/** Optional override for the sound to play when this button is pressed while Selected */
+	UPROPERTY(EditAnywhere, BlueprintReadOnly, Category = Sound, meta = (DisplayName = "Selected Pressed Sound Override"))
+	FSlateSound SelectedPressedSlateSoundOverride;
+
+	/** Optional override for the sound to play when this button is hovered while Selected */
+	UPROPERTY(EditAnywhere, BlueprintReadOnly, Category = Sound, meta = (DisplayName = "Selected Hovered Sound Override"))
+	FSlateSound SelectedHoveredSlateSoundOverride;
+
+	/** Optional override for the sound to play when this button is pressed while Locked */
+	UPROPERTY(EditAnywhere, BlueprintReadOnly, Category = Sound, meta = (DisplayName = "Locked Pressed Sound Override"))
+	FSlateSound LockedPressedSlateSoundOverride;
+
+	/** Optional override for the sound to play when this button is hovered while Locked */
+	UPROPERTY(EditAnywhere, BlueprintReadOnly, Category = Sound, meta = (DisplayName = "Locked Hovered Sound Override"))
+	FSlateSound LockedHoveredSlateSoundOverride;
+
 	/** The type of mouse action required by the user to trigger the button's 'Click' */
 	UPROPERTY(EditAnywhere, BlueprintReadOnly, Category = Style, meta = (ExposeOnSpawn = true))
 	uint8 bApplyAlphaOnDisable:1;
 
+	/**
+	 * True if this button is currently locked.
+	 * Locked button can be hovered, focused, and pressed, but the Click event will not go through.
+	 * Business logic behind it will not be executed. Designed for progressive disclosure
+	 */
+	UPROPERTY(EditAnywhere, BlueprintReadOnly, Category = Locked, meta = (ExposeOnSpawn = true))
+	uint8 bLocked:1;
+	
 	/** True if the button supports being in a "selected" state, which will update the style accordingly */
 	UPROPERTY(EditAnywhere, BlueprintReadOnly, Category = Selection, meta = (ExposeOnSpawn = true))
 	uint8 bSelectable:1;
+	
+	/** If true, the button will be selected when it receives focus. */
+    UPROPERTY(EditAnywhere, BlueprintReadOnly, Category = Interaction, meta = (ExposeOnSpawn = true, EditCondition = "bIsFocusable"))
+    uint8 bShouldFocusOnHover:1;
 
 	/** If true, the button will be selected when it receives focus. */
 	UPROPERTY(EditAnywhere, BlueprintReadOnly, Category = Selection, meta = (ExposeOnSpawn = true, EditCondition = "bSelectable"))
@@ -677,9 +722,6 @@ private:
 	/** True if this button is currently selected */
 	uint8 bSelected:1;
 
-	/** True if this button is currently locked */
-	uint8 bLocked:1;
-
 	/** True if this button is currently enabled */
 	uint8 bButtonEnabled:1;
 
@@ -784,6 +826,10 @@ private:
 	UPROPERTY()
 	FButtonStyle DisabledStyle;
 
+	/** Internally managed and applied style to use when locked */
+	UPROPERTY()
+	FButtonStyle LockedStyle;
+
 	UPROPERTY(Transient)
 	uint32 bStopDoubleClickPropagation : 1;
 
