Fix middle mouse inputs not working with CommonUI

diff --git a/Engine/Plugins/Experimental/CommonUI/Source/CommonUI/Private/Input/CommonAnalogCursor.cpp b/Engine/Plugins/Experimental/CommonUI/Source/CommonUI/Private/Input/CommonAnalogCursor.cpp
--- a/Engine/Plugins/Experimental/CommonUI/Source/CommonUI/Private/Input/CommonAnalogCursor.cpp
+++ b/Engine/Plugins/Experimental/CommonUI/Source/CommonUI/Private/Input/CommonAnalogCursor.cpp
@@ -35,7 +35,10 @@ const float ScrollDeadZone = 0.2f;
 bool IsEligibleFakeKeyPointerEvent(const FPointerEvent& PointerEvent)
 {
 	FKey EffectingButton = PointerEvent.GetEffectingButton();
-	return EffectingButton.IsMouseButton() && EffectingButton != EKeys::LeftMouseButton && EffectingButton != EKeys::RightMouseButton;
+	return EffectingButton.IsMouseButton() 
+		&& EffectingButton != EKeys::LeftMouseButton
+		&& EffectingButton != EKeys::RightMouseButton
+		&& EffectingButton != EKeys::MiddleMouseButton;
 }
 
 FCommonAnalogCursor::FCommonAnalogCursor(const UCommonUIActionRouterBase& InActionRouter)
@@ -50,6 +53,7 @@ void FCommonAnalogCursor::Initialize()
 	PointerButtonDownKeys = FSlateApplication::Get().GetPressedMouseButtons();
 	PointerButtonDownKeys.Remove(EKeys::LeftMouseButton);
 	PointerButtonDownKeys.Remove(EKeys::RightMouseButton);
+	PointerButtonDownKeys.Remove(EKeys::MiddleMouseButton);
 
 	UCommonInputSubsystem& InputSubsystem = ActionRouter.GetInputSubsystem();
 	InputSubsystem.OnInputMethodChangedNative.AddSP(this, &FCommonAnalogCursor::HandleInputMethodChanged);
@@ -345,9 +349,10 @@ bool FCommonAnalogCursor::HandleMouseButtonDownEvent(FSlateApplication& SlateApp
 		//@todo DanH: May want to make the list of "mouse buttons to treat like keys" a settings thing
 		// Mouse buttons other than the two primaries are fair game for binding as if they were normal keys
 		const FKey EffectingButton = PointerEvent.GetEffectingButton();
-		if (EffectingButton.IsMouseButton() &&
-			EffectingButton != EKeys::LeftMouseButton &&
-			EffectingButton != EKeys::RightMouseButton)
+		if (EffectingButton.IsMouseButton()
+			&& EffectingButton != EKeys::LeftMouseButton
+			&& EffectingButton != EKeys::RightMouseButton
+			&& EffectingButton != EKeys::MiddleMouseButton)
 		{
 			UGameViewportClient* ViewportClient = GetViewportClient();
 			if (TSharedPtr<SWidget> ViewportWidget = ViewportClient ? ViewportClient->GetGameViewportWidget() : nullptr)
