Updated UCommonUIActionRouterBase to have Action Domain Root nodes go through the same logic that standard root nodes go through.
This fixes issues related to widgets that are children of action domain roots from ensuring when they register input actions.
Also updated action router tree nodes to only update input config and focus if we're the active root.

diff --git a/Engine/Plugins/Experimental/CommonUI/Source/CommonInput/Public/CommonInputActionDomain.h b/Engine/Plugins/Experimental/CommonUI/Source/CommonInput/Public/CommonInputActionDomain.h
--- a/Engine/Plugins/Experimental/CommonUI/Source/CommonInput/Public/CommonInputActionDomain.h
+++ b/Engine/Plugins/Experimental/CommonUI/Source/CommonInput/Public/CommonInputActionDomain.h
@@ -50,6 +50,15 @@ public:
 	UPROPERTY(EditDefaultsOnly, Category = "Default")
 	ECommonInputEventFlowBehavior InnerBehavior = ECommonInputEventFlowBehavior::BlockIfHandled;
 
+	UPROPERTY(EditDefaultsOnly, Category = "Default" )
+	bool bUseActionDomainDesiredInputConfig;
+
+	UPROPERTY(EditDefaultsOnly, Category = "Default")
+	ECommonInputMode InputMode = ECommonInputMode::Game;
+
+	UPROPERTY(EditDefaultsOnly, Category = "Default")
+	EMouseCaptureMode MouseCaptureMode = EMouseCaptureMode::CapturePermanently;
+
 	bool ShouldBreakInnerEventFlow(bool bInputEventHandled) const;
 
 	bool ShouldBreakEventFlow(bool bDomainHadActiveRoots, bool bInputEventHandledAtLeastOnce) const;
@@ -67,4 +76,10 @@ public:
 	// Domains will receive events in ascending index order
 	UPROPERTY(EditDefaultsOnly, BlueprintReadOnly, Category = "Default")
 	TArray<TObjectPtr<UCommonInputActionDomain>> ActionDomains;
+
+	UPROPERTY(EditDefaultsOnly, Category = "Default")
+	ECommonInputMode InputMode = ECommonInputMode::Game;
+
+	UPROPERTY(EditDefaultsOnly, Category = "Default")
+	EMouseCaptureMode MouseCaptureMode = EMouseCaptureMode::CapturePermanently;
 };
\ No newline at end of file
diff --git a/Engine/Plugins/Experimental/CommonUI/Source/CommonUI/Private/Input/CommonUIActionRouterBase.cpp b/Engine/Plugins/Experimental/CommonUI/Source/CommonUI/Private/Input/CommonUIActionRouterBase.cpp
--- a/Engine/Plugins/Experimental/CommonUI/Source/CommonUI/Private/Input/CommonUIActionRouterBase.cpp
+++ b/Engine/Plugins/Experimental/CommonUI/Source/CommonUI/Private/Input/CommonUIActionRouterBase.cpp
@@ -580,6 +580,7 @@ void UCommonUIActionRouterBase::HandleRootWidgetSlateReleased(TWeakPtr<FActivata
 		{
 			//@todo DanH: This won't actually change the current config, which we may want to do with a loading screen
 			ActiveInputConfig.Reset();
+			RefreshActionDomainLeafNodeConfig();
 		}
 	}
 	else
@@ -589,6 +590,9 @@ void UCommonUIActionRouterBase::HandleRootWidgetSlateReleased(TWeakPtr<FActivata
 		{
 			FActionDomainSortedRootList& ActionDomainRootList = Pair.Value;
 			NumRemoved += ActionDomainRootList.Remove(RootNode);
+			RootNode->OnLeafmostActiveNodeChanged.Unbind();
+			RootNode->SetCanReceiveInput(false);
+			RootNode->UpdateLeafNode();
 		}
 
 		if (NumRemoved <= 0)
@@ -602,21 +606,59 @@ void UCommonUIActionRouterBase::HandleRootWidgetSlateReleased(TWeakPtr<FActivata
 void UCommonUIActionRouterBase::HandleRootNodeActivated(TWeakPtr<FActivatableTreeRoot> WeakActivatedRoot)
 {
 	FActivatableTreeRootRef ActivatedRoot = WeakActivatedRoot.Pin().ToSharedRef();
-	if (!RootNodes.Contains(ActivatedRoot))
+	UCommonActivatableWidget* NodeWidget = ActivatedRoot->GetWidget();
+	if (UCommonInputActionDomain* WidgetActionDomain = NodeWidget ? NodeWidget->GetCalculatedActionDomain() : nullptr)
 	{
+		const FActionDomainSortedRootList* ActionDomainRootList = ActionDomainRootNodes.Find(WidgetActionDomain);
+		if (ActionDomainRootList && ensure(ActionDomainRootList->Contains(ActivatedRoot)))
+		{
+			ActivatedRoot->SetCanReceiveInput(true);
+			ActivatedRoot->OnLeafmostActiveNodeChanged.BindUObject(this, &UCommonUIActionRouterBase::HandleLeafmostActiveNodeChanged);
+			
+			UCommonInputSubsystem& CommonInputSubsystem = GetInputSubsystem();
+			if (UCommonInputActionDomainTable* ActionDomainTable = CommonInputSubsystem.GetActionDomainTable())
+			{
+				// We find the first root node that isn't the activated root and bail early. 
+				// The action domains and root lists are sorted so we should end up with a node with a higher paint layer
+				// or we just stop at the activated node and make sure it's leaf node is updated.
+				for (const UCommonInputActionDomain* ActionDomain : ActionDomainTable->ActionDomains)
+				{
+					if (FActionDomainSortedRootList* SortedRootList = ActionDomainRootNodes.Find(ActionDomain))
+					{
+						for (FActivatableTreeRootRef& RootNode : SortedRootList->RootList)
+						{
+							if (RootNode->IsReceivingInput() && RootNode->DoesWidgetSupportActivationFocus())
+							{
+								if (RootNode == ActivatedRoot)
+								{
+									// This will allow us to update the leaf node config on the activated root.
+									ActivatedRoot->UpdateLeafNode();
+									OnBoundActionsUpdated().Broadcast();
+								}
+								break;
+							}
+						}
+					}
+				}
+			}
+		}
+
 		return;
 	}
 
-	if (ActivatedRoot->GetLastPaintLayer() > 0)
+	if (RootNodes.Contains(ActivatedRoot))
 	{
-		const int32 CurrentRootLayer = ActiveRootNode ? ActiveRootNode->GetLastPaintLayer() : INDEX_NONE;
-		if (ActivatedRoot->GetLastPaintLayer() > CurrentRootLayer)
+		if (ActivatedRoot->GetLastPaintLayer() > 0)
 		{
-			// Ensure we have a local player so the action router local player subsystem will handle root change
-			const ULocalPlayer* LocalPlayer = GetLocalPlayer();
-			if (LocalPlayer && LocalPlayer->ViewportClient)
+			const int32 CurrentRootLayer = ActiveRootNode ? ActiveRootNode->GetLastPaintLayer() : INDEX_NONE;
+			if (ActivatedRoot->GetLastPaintLayer() > CurrentRootLayer)
 			{
-				SetActiveRoot(ActivatedRoot);
+				// Ensure we have a local player so the action router local player subsystem will handle root change
+				const ULocalPlayer* LocalPlayer = GetLocalPlayer();
+				if (LocalPlayer && LocalPlayer->ViewportClient)
+				{
+					SetActiveRoot(ActivatedRoot);
+				}
 			}
 		}
 	}
@@ -624,10 +666,21 @@ void UCommonUIActionRouterBase::HandleRootNodeActivated(TWeakPtr<FActivatableTre
 
 void UCommonUIActionRouterBase::HandleRootNodeDeactivated(TWeakPtr<FActivatableTreeRoot> WeakDeactivatedRoot)
 {
-	if (ActiveRootNode && ActiveRootNode == WeakDeactivatedRoot.Pin())
+	FActivatableTreeRootPtr DeactivatedRoot = WeakDeactivatedRoot.Pin();
+	UCommonActivatableWidget* NodeWidget = DeactivatedRoot ? DeactivatedRoot->GetWidget() : nullptr;
+	UCommonInputActionDomain* ActionDomain = NodeWidget ? NodeWidget->GetCalculatedActionDomain() : nullptr;
+	
+	if (ActionDomain || (ActiveRootNode && ActiveRootNode == DeactivatedRoot))
 	{
+		// Action domain root nodes will need to recalibrate during deactivation since we're not relying on 
+		// timing to fix up the input config as RootNodes does.
+		RefreshActionDomainLeafNodeConfig();
+
 		// Reset the active root widget - we'll re-establish it on the next tick
-		SetActiveRoot(nullptr);
+		if (ActiveRootNode.IsValid())
+		{
+			SetActiveRoot(nullptr);
+		}
 	}
 }
 
@@ -777,13 +830,7 @@ void UCommonUIActionRouterBase::ProcessRebuiltWidgets()
 	{
 		if (RebuiltWidget.IsValid() && RebuiltWidget->GetCachedWidget())
 		{
-			if (UCommonInputActionDomain* ActionDomain = RebuiltWidget->GetCalculatedActionDomain())
-			{
-				FActivatableTreeRootRef RootNode = FActivatableTreeRoot::Create(*this, *RebuiltWidget);
-				RebuiltWidget->OnSlateReleased().AddUObject(this, &UCommonUIActionRouterBase::HandleRootWidgetSlateReleased, TWeakPtr<FActivatableTreeRoot>(RootNode));
-				ActionDomainRootNodes.FindOrAdd(ActionDomain).Add(RootNode);
-			}
-			else if (UCommonActivatableWidget* ActivatableParent = !RebuiltWidget->IsModal() ? ::FindOwningActivatable(*RebuiltWidget) : nullptr)
+			if (UCommonActivatableWidget* ActivatableParent = !RebuiltWidget->IsModal() ? ::FindOwningActivatable(*RebuiltWidget) : nullptr)
 			{
 				WidgetsByDirectParent.FindOrAdd(ActivatableParent).Add(RebuiltWidget.Get());
 			}
@@ -805,7 +852,15 @@ void UCommonUIActionRouterBase::ProcessRebuiltWidgets()
 		RootNode->OnDeactivated.BindUObject(this, &UCommonUIActionRouterBase::HandleRootNodeDeactivated, WeakRoot);
 		RootWidget->OnSlateReleased().AddUObject(this, &UCommonUIActionRouterBase::HandleRootWidgetSlateReleased, WeakRoot);
 
-		RootNodes.Add(RootNode);
+		UCommonInputActionDomain* ActionDomain = RootWidget->GetCalculatedActionDomain();
+		if (ActionDomain)
+		{
+			ActionDomainRootNodes.FindOrAdd(ActionDomain).Add(RootNode);
+		}
+		else
+		{
+			RootNodes.Add(RootNode);
+		}
 
 		AssembleTreeRecursive(RootNode, WidgetsByDirectParent);
 
@@ -1055,6 +1110,7 @@ void UCommonUIActionRouterBase::SetActiveRoot(FActivatableTreeRootPtr NewActiveR
 	{
 		ActiveRootNode->OnLeafmostActiveNodeChanged.Unbind();
 		ActiveRootNode->SetCanReceiveInput(false);
+		ActiveRootNode->UpdateLeafNode();
 	}
 
 	if (bForceResetActiveRoot || !bIsActivatableTreeEnabled)
@@ -1069,6 +1125,7 @@ void UCommonUIActionRouterBase::SetActiveRoot(FActivatableTreeRootPtr NewActiveR
 		if (NewActiveRoot)
 		{
 			NewActiveRoot->SetCanReceiveInput(true);
+			NewActiveRoot->UpdateLeafNode();
 			NewActiveRoot->OnLeafmostActiveNodeChanged.BindUObject(this, &UCommonUIActionRouterBase::HandleLeafmostActiveNodeChanged);
 		}
 	}
@@ -1081,6 +1138,18 @@ void UCommonUIActionRouterBase::SetForceResetActiveRoot(bool bInForceResetActive
 	bForceResetActiveRoot = bInForceResetActiveRoot;
 }
 
+void UCommonUIActionRouterBase::UpdateLeafNodeAndConfig(FActivatableTreeRootPtr DesiredRoot, FActivatableTreeNodePtr DesiredLeafNode)
+{
+	if (DesiredRoot.IsValid())
+	{
+		// We're updating both the leaf node and it's config if we're the active root.
+		if (!DesiredRoot->UpdateLeafmostActiveNode(DesiredLeafNode, DesiredRoot == ActiveRootNode))
+		{
+			UE_LOG(LogUIActionRouter, Warning, TEXT("LeafmostActiveNode not updated."));
+		}
+	}
+}
+
 UCommonUIActionRouterBase::FPendingWidgetRegistration& UCommonUIActionRouterBase::GetOrCreatePendingRegistration(const UWidget& Widget)
 {
 	if (FPendingWidgetRegistration* ExistingEntry = PendingWidgetRegistrations.FindByKey(Widget))
@@ -1123,9 +1192,10 @@ FActivatableTreeNodePtr UCommonUIActionRouterBase::FindNode(const UCommonActivat
 			{
 				for (const FActivatableTreeRootRef& RootNode : Pair.Value.RootList)
 				{
-					if (Widget == RootNode->GetWidget())
+					FoundNode = FindNodeRecursive(RootNode, *Widget);
+					if (FoundNode.IsValid())
 					{
-						return RootNode;
+						break;
 					}
 				}
 			}
@@ -1210,6 +1280,35 @@ void UCommonUIActionRouterBase::SetActiveUIInputConfig(const FUIInputConfig& New
 	ApplyUIInputConfig(NewConfig, !ActiveInputConfig.IsSet());
 }
 
+void UCommonUIActionRouterBase::RefreshActionDomainLeafNodeConfig()
+{
+	UCommonInputSubsystem& CommonInputSubsystem = GetInputSubsystem();
+	if (UCommonInputActionDomainTable* ActionDomainTable = CommonInputSubsystem.GetActionDomainTable())
+	{
+		for (const UCommonInputActionDomain* ActionDomain : ActionDomainTable->ActionDomains)
+		{
+			if (FActionDomainSortedRootList* SortedRootList = ActionDomainRootNodes.Find(ActionDomain))
+			{
+				for (FActivatableTreeRootRef& RootNode : SortedRootList->RootList)
+				{
+					// only root nodes that are actively receiving input and supports widget activation focus
+					// should update leaf nodes and have input config applied. This will also update focus.
+					if (RootNode->IsReceivingInput() && RootNode->DoesWidgetSupportActivationFocus())
+					{
+						if (!RootNode->UpdateLeafmostActiveNode(RootNode))
+						{
+							RootNode->ApplyLeafmostNodeConfig();
+						}
+						return;
+					}
+				}
+			}
+		}
+
+		SetActiveUIInputConfig(FUIInputConfig(ActionDomainTable->InputMode, ActionDomainTable->MouseCaptureMode));
+	}
+}
+
 void UCommonUIActionRouterBase::ApplyUIInputConfig(const FUIInputConfig& NewConfig, bool bForceRefresh)
 {
 	if (bForceRefresh || NewConfig != ActiveInputConfig.GetValue())
@@ -1470,7 +1569,7 @@ int32 UCommonUIActionRouterBase::FActionDomainSortedRootList::Remove(FActivatabl
 	return RootList.Remove(RootNode);
 }
 
-bool UCommonUIActionRouterBase::FActionDomainSortedRootList::Contains(FActivatableTreeRootRef RootNode)
+bool UCommonUIActionRouterBase::FActionDomainSortedRootList::Contains(FActivatableTreeRootRef RootNode) const
 {
 	return RootList.Contains(RootNode);
 }
diff --git a/Engine/Plugins/Experimental/CommonUI/Source/CommonUI/Private/Input/UIActionRouterTypes.cpp b/Engine/Plugins/Experimental/CommonUI/Source/CommonUI/Private/Input/UIActionRouterTypes.cpp
--- a/Engine/Plugins/Experimental/CommonUI/Source/CommonUI/Private/Input/UIActionRouterTypes.cpp
+++ b/Engine/Plugins/Experimental/CommonUI/Source/CommonUI/Private/Input/UIActionRouterTypes.cpp
@@ -815,14 +815,28 @@ int32 FActivatableTreeNode::GetLastPaintLayer() const
 TOptional<FUIInputConfig> FActivatableTreeNode::FindDesiredInputConfig() const
 {
 	TOptional<FUIInputConfig> DesiredConfig = ensure(RepresentedWidget.IsValid()) ? RepresentedWidget->GetDesiredInputConfig() : TOptional<FUIInputConfig>();
+	TOptional<FUIInputConfig> ActionDomainDesiredConfig = FindDesiredActionDomainInputConfig();
+
+	if (ActionDomainDesiredConfig.IsSet() && !DesiredConfig.IsSet())
+	{
+		return ActionDomainDesiredConfig;
+	}
+	
 	if (!DesiredConfig.IsSet() && Parent.IsValid())
 	{
 		DesiredConfig = Parent.Pin()->FindDesiredInputConfig();
 	}
-
+	
 	return DesiredConfig;
 }
 
+TOptional<FUIInputConfig> FActivatableTreeNode::FindDesiredActionDomainInputConfig() const
+{
+	UCommonInputActionDomain* ActionDomain = ensure(RepresentedWidget.IsValid()) ? RepresentedWidget->GetCalculatedActionDomain() : nullptr;
+	const bool bHasActionDomainConfig = ActionDomain && ActionDomain->bUseActionDomainDesiredInputConfig;
+	return  bHasActionDomainConfig ? FUIInputConfig(ActionDomain->InputMode, ActionDomain->MouseCaptureMode) : TOptional<FUIInputConfig>();
+}
+
 FUICameraConfig FActivatableTreeNode::FindDesiredCameraConfig() const
 {
 	TOptional<FUICameraConfig> DesiredConfig = ensure(RepresentedWidget.IsValid()) ? RepresentedWidget->GetDesiredCameraConfig() : TOptional<FUICameraConfig>();
@@ -1039,7 +1053,8 @@ void FActivatableTreeNode::HandleWidgetDeactivated()
 				{
 					NearestActiveParent = NearestActiveParent->GetParentNode();
 				}
-				GetRoot()->UpdateLeafmostActiveNode(NearestActiveParent);
+
+				GetActionRouter().UpdateLeafNodeAndConfig(GetRoot(), NearestActiveParent);
 			}
 		}
 	}
@@ -1123,27 +1138,23 @@ FActivatableTreeRootRef FActivatableTreeRoot::Create(UCommonUIActionRouterBase&
 	return NewRoot;
 }
 
-void FActivatableTreeRoot::SetCanReceiveInput(bool bInCanReceiveInput)
+void FActivatableTreeRoot::UpdateLeafNode()
 {
-	if (CanReceiveInput() != bInCanReceiveInput)
+	if (!CanReceiveInput())
 	{
-		FActivatableTreeNode::SetCanReceiveInput(bInCanReceiveInput);
-		if (!bInCanReceiveInput)
+		if (LeafmostActiveNode.IsValid())
 		{
-			if (LeafmostActiveNode.IsValid())
-			{
-				LeafmostActiveNode.Pin()->CacheFocusRestorationTarget();
-				UpdateLeafmostActiveNode(nullptr);
-			}
+			LeafmostActiveNode.Pin()->CacheFocusRestorationTarget();
+			UpdateLeafmostActiveNode(nullptr);
 		}
-		else if (ensure(IsWidgetActivated()))
+	}
+	else if (ensure(IsWidgetActivated()))
+	{
+		if (!UpdateLeafmostActiveNode(SharedThis(this)))
 		{
-			if (!UpdateLeafmostActiveNode(SharedThis(this)))
-			{
-				// Our leafmost active node didn't change (good!), so we make sure apply its desired config
-				// We only bother when the update doesn't do anything to avoid calling Apply twice.
-				ApplyLeafmostNodeConfig();
-			}
+			// Our leafmost active node didn't change (good!), so we make sure apply its desired config
+			// We only bother when the update doesn't do anything to avoid calling Apply twice.
+			ApplyLeafmostNodeConfig();
 		}
 	}
 }
@@ -1155,7 +1166,7 @@ TArray<const UWidget*> FActivatableTreeRoot::GatherScrollRecipients() const
 	return AllScrollRecipients;
 }
 
-bool FActivatableTreeRoot::UpdateLeafmostActiveNode(FActivatableTreeNodePtr BaseCandidateNode)
+bool FActivatableTreeRoot::UpdateLeafmostActiveNode(FActivatableTreeNodePtr BaseCandidateNode, bool bInApplyConfig)
 {
 	bool bValidBaseCandidate = !BaseCandidateNode || BaseCandidateNode->IsReceivingInput();
 	if (!bValidBaseCandidate)
@@ -1190,7 +1201,10 @@ bool FActivatableTreeRoot::UpdateLeafmostActiveNode(FActivatableTreeNodePtr Base
 		}
 
 		LeafmostActiveNode = NewLeafmostNode;
-		ApplyLeafmostNodeConfig();
+		if (bInApplyConfig)
+		{
+			ApplyLeafmostNodeConfig();
+		}
 
 		if (LeafmostActiveNode.IsValid())
 		{
diff --git a/Engine/Plugins/Experimental/CommonUI/Source/CommonUI/Private/Input/UIActionRouterTypes.h b/Engine/Plugins/Experimental/CommonUI/Source/CommonUI/Private/Input/UIActionRouterTypes.h
--- a/Engine/Plugins/Experimental/CommonUI/Source/CommonUI/Private/Input/UIActionRouterTypes.h
+++ b/Engine/Plugins/Experimental/CommonUI/Source/CommonUI/Private/Input/UIActionRouterTypes.h
@@ -190,8 +190,11 @@ public:
 
 	int32 GetLastPaintLayer() const;
 	TOptional<FUIInputConfig> FindDesiredInputConfig() const;
+	TOptional<FUIInputConfig> FindDesiredActionDomainInputConfig() const;
 	FUICameraConfig FindDesiredCameraConfig() const;
-
+	
+	void SetCanReceiveInput(bool bInCanReceiveInput);
+	
 	void AddScrollRecipient(const UWidget& ScrollRecipient);
 	void RemoveScrollRecipient(const UWidget& ScrollRecipient);
 	void AddInputPreprocessor(const TSharedRef<IInputProcessor>& InputPreprocessor, int32 DesiredIndex);
@@ -204,11 +207,9 @@ protected:
 	FActivatableTreeNode(UCommonUIActionRouterBase& OwningRouter, UCommonActivatableWidget& ActivatableWidget, const FActivatableTreeNodeRef& InParent);
 
 	virtual bool IsWidgetReachableForInput(const UWidget* Widget) const override;
-
-	virtual void Init();
-	virtual void SetCanReceiveInput(bool bInCanReceiveInput);
 	
 	bool CanReceiveInput() const { return bCanReceiveInput; }
+	virtual void Init();	
 	FActivatableTreeRootRef GetRoot() const;
 
 	void AppendValidScrollRecipients(TArray<const UWidget*>& AllScrollRecipients) const;
@@ -260,11 +261,11 @@ class COMMONUI_API FActivatableTreeRoot : public FActivatableTreeNode
 public:
 	static FActivatableTreeRootRef Create(UCommonUIActionRouterBase& OwningRouter, UCommonActivatableWidget& ActivatableWidget);
 	
-	virtual void SetCanReceiveInput(bool bInCanReceiveInput) override;
+	void UpdateLeafNode();
 
 	TArray<const UWidget*> GatherScrollRecipients() const;
 
-	bool UpdateLeafmostActiveNode(FActivatableTreeNodePtr BaseCandidateNode);
+	bool UpdateLeafmostActiveNode(FActivatableTreeNodePtr BaseCandidateNode, bool bInApplyConfig = true);
 
 	void DebugDump(FString& OutputStr, bool bIncludeActions, bool bIncludeChildren, bool bIncludeInactive) const;
 
@@ -274,6 +275,8 @@ public:
 
 	void RefreshCachedRestorationTarget();
 
+	void ApplyLeafmostNodeConfig();
+
 protected:
 	virtual void Init() override;
 
@@ -284,8 +287,6 @@ private:
 
 	void HandleInputMethodChanged(ECommonInputType InputMethod);
 
-	void ApplyLeafmostNodeConfig();
-
 	void HandleRequestRefreshLeafmostFocus();
 
 	// WeakPtr because the root itself can be the primary active node - results in a circular ref leak using a full SharedPtr here
diff --git a/Engine/Plugins/Experimental/CommonUI/Source/CommonUI/Public/Input/CommonUIActionRouterBase.h b/Engine/Plugins/Experimental/CommonUI/Source/CommonUI/Public/Input/CommonUIActionRouterBase.h
--- a/Engine/Plugins/Experimental/CommonUI/Source/CommonUI/Public/Input/CommonUIActionRouterBase.h
+++ b/Engine/Plugins/Experimental/CommonUI/Source/CommonUI/Public/Input/CommonUIActionRouterBase.h
@@ -125,6 +125,10 @@ protected:
 	virtual void SetActiveRoot(FActivatableTreeRootPtr NewActiveRoot);
 	void SetForceResetActiveRoot(bool bInForceResetActiveRoot);
 
+	void UpdateLeafNodeAndConfig(FActivatableTreeRootPtr DesiredRoot, FActivatableTreeNodePtr DesiredLeafNode);
+
+	void RefreshActionDomainLeafNodeConfig();
+
 	bool bIsActivatableTreeEnabled = true;
 
 	TSharedPtr<FCommonAnalogCursor> AnalogCursor;
@@ -218,7 +222,7 @@ private:
 		int32 Remove(FActivatableTreeRootRef RootNode);
 
 		// Trivial Contains check
-		bool Contains(FActivatableTreeRootRef RootNode);
+		bool Contains(FActivatableTreeRootRef RootNode) const;
 	};
 
 	TMap<TObjectPtr<UCommonInputActionDomain>, FActionDomainSortedRootList> ActionDomainRootNodes;
