[NO NON-PII COMMIT MESSAGE]

diff --git a/Engine/Plugins/Experimental/CommonUI/Source/CommonInput/Public/CommonInputActionDomain.h b/Engine/Plugins/Experimental/CommonUI/Source/CommonInput/Public/CommonInputActionDomain.h
--- a/Engine/Plugins/Experimental/CommonUI/Source/CommonInput/Public/CommonInputActionDomain.h
+++ b/Engine/Plugins/Experimental/CommonUI/Source/CommonInput/Public/CommonInputActionDomain.h
@@ -3,6 +3,7 @@
 #pragma once
 
 #include "Engine/DataAsset.h"
+#include "Engine/EngineBaseTypes.h"
 #include "Types/ISlateMetaData.h"
 #include "CommonInputModeTypes.h"
 
