#rnx

diff --git a/Engine/Plugins/Experimental/CommonUI/Source/CommonUI/Private/Input/CommonUIActionRouterBase.cpp b/Engine/Plugins/Experimental/CommonUI/Source/CommonUI/Private/Input/CommonUIActionRouterBase.cpp
--- a/Engine/Plugins/Experimental/CommonUI/Source/CommonUI/Private/Input/CommonUIActionRouterBase.cpp
+++ b/Engine/Plugins/Experimental/CommonUI/Source/CommonUI/Private/Input/CommonUIActionRouterBase.cpp
@@ -1282,30 +1282,32 @@ void UCommonUIActionRouterBase::SetActiveUIInputConfig(const FUIInputConfig& New
 
 void UCommonUIActionRouterBase::RefreshActionDomainLeafNodeConfig()
 {
-	UCommonInputSubsystem& CommonInputSubsystem = GetInputSubsystem();
-	if (UCommonInputActionDomainTable* ActionDomainTable = CommonInputSubsystem.GetActionDomainTable())
+	if (UCommonInputSubsystem* CommonInputSubsystem = GetLocalPlayer() ? GetLocalPlayer()->GetSubsystem<UCommonInputSubsystem>() : nullptr)
 	{
-		for (const UCommonInputActionDomain* ActionDomain : ActionDomainTable->ActionDomains)
+		if (UCommonInputActionDomainTable* ActionDomainTable = CommonInputSubsystem->GetActionDomainTable())
 		{
-			if (FActionDomainSortedRootList* SortedRootList = ActionDomainRootNodes.Find(ActionDomain))
+			for (const UCommonInputActionDomain* ActionDomain : ActionDomainTable->ActionDomains)
 			{
-				for (FActivatableTreeRootRef& RootNode : SortedRootList->RootList)
+				if (FActionDomainSortedRootList* SortedRootList = ActionDomainRootNodes.Find(ActionDomain))
 				{
-					// only root nodes that are actively receiving input and supports widget activation focus
-					// should update leaf nodes and have input config applied. This will also update focus.
-					if (RootNode->IsReceivingInput() && RootNode->DoesWidgetSupportActivationFocus())
+					for (FActivatableTreeRootRef& RootNode : SortedRootList->RootList)
 					{
-						if (!RootNode->UpdateLeafmostActiveNode(RootNode))
+						// only root nodes that are actively receiving input and supports widget activation focus
+						// should update leaf nodes and have input config applied. This will also update focus.
+						if (RootNode->IsReceivingInput() && RootNode->DoesWidgetSupportActivationFocus())
 						{
-							RootNode->ApplyLeafmostNodeConfig();
+							if (!RootNode->UpdateLeafmostActiveNode(RootNode))
+							{
+								RootNode->ApplyLeafmostNodeConfig();
+							}
+							return;
 						}
-						return;
 					}
 				}
 			}
-		}
 
-		SetActiveUIInputConfig(FUIInputConfig(ActionDomainTable->InputMode, ActionDomainTable->MouseCaptureMode));
+			SetActiveUIInputConfig(FUIInputConfig(ActionDomainTable->InputMode, ActionDomainTable->MouseCaptureMode));
+		}
 	}
 }
 
