Allow you to override the GetIcon function on UCommonActionWidget. This allows users to make a subclass of the widget that will ask Enhanced Input about what FKey the player has remapped an action.

#rnx

diff --git a/Engine/Plugins/Experimental/CommonUI/Source/CommonUI/Public/CommonActionWidget.h b/Engine/Plugins/Experimental/CommonUI/Source/CommonUI/Public/CommonActionWidget.h
--- a/Engine/Plugins/Experimental/CommonUI/Source/CommonUI/Public/CommonActionWidget.h
+++ b/Engine/Plugins/Experimental/CommonUI/Source/CommonUI/Public/CommonActionWidget.h
@@ -32,7 +32,7 @@ public:
 	/** End UWidet */
 	
 	UFUNCTION(BlueprintCallable, Category = CommonActionWidget)
-	FSlateBrush GetIcon() const;
+	virtual FSlateBrush GetIcon() const;
 
 	UFUNCTION(BlueprintCallable, Category = CommonActionWidget)
 	FText GetDisplayText() const;
