Wrapped content with Scale box in CommonActionWidget (to support different sets of sizes for icon prompts). Changed order for hold progress material - put it behind the prompt icon.

diff --git a/Engine/Plugins/Experimental/CommonUI/Source/CommonUI/Private/CommonActionWidget.cpp b/Engine/Plugins/Experimental/CommonUI/Source/CommonUI/Private/CommonActionWidget.cpp
--- a/Engine/Plugins/Experimental/CommonUI/Source/CommonUI/Private/CommonActionWidget.cpp
+++ b/Engine/Plugins/Experimental/CommonUI/Source/CommonUI/Private/CommonActionWidget.cpp
@@ -57,28 +57,35 @@ TSharedRef<SWidget> UCommonActionWidget::RebuildWidget()
 		.HAlign(HAlign_Center)
 		.VAlign(VAlign_Center);
 
-	MyKeyBox->SetContent(	
-		SNew(SOverlay)
-		+ SOverlay::Slot()
+	MyKeyBox->SetContent(
+		SNew(SScaleBox)
 		.HAlign(HAlign_Center)
 		.VAlign(VAlign_Center)
+		.Stretch(IconStretch)
+		.UserSpecifiedScale(UserSpecifiedScale)
 		[
-			SAssignNew(MyIconRim, SImage)
-			.Image(&IconRimBrush)
-		]
-		+ SOverlay::Slot()
-		.HAlign(HAlign_Center)
-		.VAlign(VAlign_Center)
-		[
-			SAssignNew(MyIcon, SImage)
-			.Image(&Icon)
-		]
-		+ SOverlay::Slot()
-		.HAlign(HAlign_Center)
-		.VAlign(VAlign_Center)
-		[
-			SAssignNew(MyProgressImage, SImage)
-			.Image(&ProgressMaterialBrush)
+			SNew(SOverlay)
+			+ SOverlay::Slot()
+			.HAlign(HAlign_Center)
+			.VAlign(VAlign_Center)
+			[
+				SAssignNew(MyIconRim, SImage)
+				.Image(&IconRimBrush)
+			]
+			+ SOverlay::Slot()
+			.HAlign(HAlign_Center)
+			.VAlign(VAlign_Center)
+			[
+				SAssignNew(MyProgressImage, SImage)
+				.Image(&ProgressMaterialBrush)
+			]
+			+ SOverlay::Slot()
+			.HAlign(HAlign_Center)
+			.VAlign(VAlign_Center)
+			[
+				SAssignNew(MyIcon, SImage)
+				.Image(&Icon)
+			]
 		]);
 	
 	return MyKeyBox.ToSharedRef();
@@ -89,6 +96,7 @@ void UCommonActionWidget::ReleaseSlateResources(bool bReleaseChildren)
 	MyProgressImage.Reset();
 	MyIcon.Reset();
 	MyKeyBox.Reset();
+	MyScaleBox.Reset();
 	
 	ListenToInputMethodChanged(false);
 	Super::ReleaseSlateResources(bReleaseChildren);
diff --git a/Engine/Plugins/Experimental/CommonUI/Source/CommonUI/Public/CommonActionWidget.h b/Engine/Plugins/Experimental/CommonUI/Source/CommonUI/Public/CommonActionWidget.h
--- a/Engine/Plugins/Experimental/CommonUI/Source/CommonUI/Public/CommonActionWidget.h
+++ b/Engine/Plugins/Experimental/CommonUI/Source/CommonUI/Public/CommonActionWidget.h
@@ -4,6 +4,7 @@
 
 #include "CommonUITypes.h"
 #include "Input/UIActionBindingHandle.h"
+#include "Widgets/Layout/SScaleBox.h"
 
 #include "CommonActionWidget.generated.h"
 
@@ -68,6 +69,12 @@ public:
 
 	UPROPERTY(EditAnywhere, BlueprintReadOnly, Category = CommonActionWidget)
 	FSlateBrush IconRimBrush;
+	
+	UPROPERTY(EditAnywhere, BlueprintReadOnly, Category = CommonActionWidget)
+	TEnumAsByte<EStretch::Type> IconStretch = EStretch::UserSpecified;
+	
+	UPROPERTY(EditAnywhere, BlueprintReadOnly, Category = CommonActionWidget)
+	float UserSpecifiedScale = 1.0f;
 
 #if WITH_EDITOR
 	virtual const FText GetPaletteCategory() override;
@@ -113,6 +120,8 @@ protected:
 	
 	TSharedPtr<SBox> MyKeyBox;
 
+	TSharedPtr<SScaleBox> MyScaleBox;
+
 	TSharedPtr<SImage> MyIcon;
 
 	TSharedPtr<SImage> MyProgressImage;
