SCommonButtonTableRow:
- Added null checks when retrieving Item.
- Use GetItemForThis to retrieve Item.

Matches SObjectTableRow.

diff --git a/Engine/Plugins/Experimental/CommonUI/Source/CommonUI/Public/SCommonButtonTableRow.h b/Engine/Plugins/Experimental/CommonUI/Source/CommonUI/Public/SCommonButtonTableRow.h
--- a/Engine/Plugins/Experimental/CommonUI/Source/CommonUI/Public/SCommonButtonTableRow.h
+++ b/Engine/Plugins/Experimental/CommonUI/Source/CommonUI/Public/SCommonButtonTableRow.h
@@ -113,15 +113,21 @@ private:
 	void HandleButtonClicked()
 	{
 		TSharedRef<ITypedTableView<ItemType>> OwnerTable = this->OwnerTablePtr.Pin().ToSharedRef();
-		const ItemType& MyItem = *OwnerTable->Private_ItemFromWidget(this);
-		OwnerTable->Private_OnItemClicked(MyItem);
+
+		if (const ItemType* MyItemPtr = this->GetItemForThis(OwnerTable))
+		{
+			OwnerTable->Private_OnItemClicked(*MyItemPtr);
+		}
 	}
 
 	void HandleButtonDoubleClicked()
 	{
 		TSharedRef<ITypedTableView<ItemType>> OwnerTable = this->OwnerTablePtr.Pin().ToSharedRef();
-		const ItemType& MyItem = *OwnerTable->Private_ItemFromWidget(this);
-		OwnerTable->Private_OnItemDoubleClicked(MyItem);
+
+		if (const ItemType* MyItemPtr = this->GetItemForThis(OwnerTable))
+		{
+			OwnerTable->Private_OnItemDoubleClicked(*MyItemPtr);
+		}
 
 		// Do we want to behave differently or normally on double clicks?
 		// HandleButtonClicked();
@@ -143,20 +149,24 @@ private:
 		if (ensure(SelectionMode != ESelectionMode::None) && bIsButtonSelected != this->IsItemSelected())
 		{
 			TSharedRef<ITypedTableView<ItemType>> OwnerTable = this->OwnerTablePtr.Pin().ToSharedRef();
-			const ItemType& MyItem = *OwnerTable->Private_ItemFromWidget(this);
-			if (bIsButtonSelected)
+
+			if (const ItemType* MyItemPtr = this->GetItemForThis(OwnerTable))
 			{
-				if (SelectionMode != ESelectionMode::Multi)
+				if (bIsButtonSelected)
 				{
-					OwnerTable->Private_ClearSelection();
+					if (SelectionMode != ESelectionMode::Multi)
+					{
+						OwnerTable->Private_ClearSelection();
+					}
+					OwnerTable->Private_SetItemSelection(*MyItemPtr, true, true);
 				}
-				OwnerTable->Private_SetItemSelection(MyItem, true, true);
-			}
-			else
-			{
-				OwnerTable->Private_SetItemSelection(MyItem, false, true);
+				else
+				{
+					OwnerTable->Private_SetItemSelection(*MyItemPtr, false, true);
+				}
+
+				OwnerTable->Private_SignalSelectionChanged(ESelectInfo::Direct);
 			}
-			OwnerTable->Private_SignalSelectionChanged(ESelectInfo::Direct);
 		}
 	}
-};
\ No newline at end of file
+};
