CommonUI - Adding extra checks and some debug messages.  Trivial.

diff --git a/Engine/Plugins/Experimental/CommonUI/Source/CommonUI/Private/Input/CommonUIActionRouterBase.cpp b/Engine/Plugins/Experimental/CommonUI/Source/CommonUI/Private/Input/CommonUIActionRouterBase.cpp
--- a/Engine/Plugins/Experimental/CommonUI/Source/CommonUI/Private/Input/CommonUIActionRouterBase.cpp
+++ b/Engine/Plugins/Experimental/CommonUI/Source/CommonUI/Private/Input/CommonUIActionRouterBase.cpp
@@ -1500,6 +1500,7 @@ public:
 	{
 		if (World == nullptr)
 		{
+			UE_LOG(LogUIActionRouter, Error, TEXT("No World, unable to run CommonUI.DumpInputConfig"));
 			return;
 		}
 
@@ -1510,6 +1511,12 @@ public:
 		check(MouseCaptureModeEnum);
 
 		UGameInstance* GameInstance = World->GetGameInstance();
+		if (!GameInstance)
+		{
+			UE_LOG(LogUIActionRouter, Error, TEXT("No GameInstance, unable to run CommonUI.DumpInputConfig"));
+			return;
+		}
+
 		FString OutStr;
 		const TArray<ULocalPlayer*> LocalPlayers = GameInstance->GetLocalPlayers();
 		for (int32 i = 0; i < LocalPlayers.Num(); ++i)
