Fix for potential crash in UCommonUIVisibilitySubsystem when GameInstance isn't initialized

diff --git a/Engine/Plugins/Experimental/CommonUI/Source/CommonUI/Private/CommonUIVisibilitySubsystem.cpp b/Engine/Plugins/Experimental/CommonUI/Source/CommonUI/Private/CommonUIVisibilitySubsystem.cpp
--- a/Engine/Plugins/Experimental/CommonUI/Source/CommonUI/Private/CommonUIVisibilitySubsystem.cpp
+++ b/Engine/Plugins/Experimental/CommonUI/Source/CommonUI/Private/CommonUIVisibilitySubsystem.cpp
@@ -41,7 +41,8 @@ UCommonUIVisibilitySubsystem::UCommonUIVisibilitySubsystem()
 
 bool UCommonUIVisibilitySubsystem::ShouldCreateSubsystem(UObject* Outer) const
 {
-	if (!CastChecked<ULocalPlayer>(Outer)->GetGameInstance()->IsDedicatedServerInstance())
+	const UGameInstance* GameInstance = CastChecked<ULocalPlayer>(Outer)->GetGameInstance();
+	if (GameInstance && !GameInstance->IsDedicatedServerInstance())
 	{
 		TArray<UClass*> ChildClasses;
 		GetDerivedClasses(GetClass(), ChildClasses, false);
