Fix for CommonAnalogCursor attempting to hide cursor when using Mouse for Touch Input

diff --git a/Engine/Plugins/Experimental/CommonUI/Source/CommonUI/Private/Input/CommonAnalogCursor.cpp b/Engine/Plugins/Experimental/CommonUI/Source/CommonUI/Private/Input/CommonAnalogCursor.cpp
--- a/Engine/Plugins/Experimental/CommonUI/Source/CommonUI/Private/Input/CommonAnalogCursor.cpp
+++ b/Engine/Plugins/Experimental/CommonUI/Source/CommonUI/Private/Input/CommonAnalogCursor.cpp
@@ -339,7 +339,7 @@ bool FCommonAnalogCursor::HandleMouseButtonDownEvent(FSlateApplication& SlateApp
 	{
 #if UE_COMMONUI_PLATFORM_SUPPORTS_TOUCH	
 		// Some platforms don't register as switching its input type, so detect touch input here to hide the cursor.
-		if (PointerEvent.IsTouchEvent())
+		if (PointerEvent.IsTouchEvent() && ShouldHideCursor())
 		{
 			//ClearCenterWidget();
 			HideCursor();
@@ -526,6 +526,18 @@ bool FCommonAnalogCursor::IsUsingGamepad() const
 	return ActiveInputMethod == ECommonInputType::Gamepad;
 }
 
+bool FCommonAnalogCursor::ShouldHideCursor() const
+{
+	bool bUsingMouseForTouch = FSlateApplication::Get().IsFakingTouchEvents();
+	const ULocalPlayer& LocalPlayer = *ActionRouter.GetLocalPlayerChecked();
+	if (UGameViewportClient* GameViewportClient = LocalPlayer.ViewportClient)
+	{
+		bUsingMouseForTouch |= GameViewportClient->GetUseMouseForTouch();
+	}
+
+	return !bUsingMouseForTouch;
+}
+
 void FCommonAnalogCursor::HideCursor()
 {
 	const TSharedPtr<FSlateUser> SlateUser = FSlateApplication::Get().GetUser(GetOwnerUserIndex());
diff --git a/Engine/Plugins/Experimental/CommonUI/Source/CommonUI/Public/Input/CommonAnalogCursor.h b/Engine/Plugins/Experimental/CommonUI/Source/CommonUI/Public/Input/CommonAnalogCursor.h
--- a/Engine/Plugins/Experimental/CommonUI/Source/CommonUI/Public/Input/CommonAnalogCursor.h
+++ b/Engine/Plugins/Experimental/CommonUI/Source/CommonUI/Public/Input/CommonAnalogCursor.h
@@ -83,6 +83,11 @@ protected:
 	
 	bool IsUsingGamepad() const;
 
+
+	bool ShouldHideCursor() const;
+
+
+
 	// Knowingly unorthodox member reference to a UObject - ok because we are a subobject of the owning router and will never outlive it
 	const UCommonUIActionRouterBase& ActionRouter;
 	ECommonInputType ActiveInputMethod;
