#rnx
#test PIE: Ran through some test cases

diff --git a/Engine/Plugins/Experimental/CommonUI/Source/CommonUI/Private/CommonButtonBase.cpp b/Engine/Plugins/Experimental/CommonUI/Source/CommonUI/Private/CommonButtonBase.cpp
--- a/Engine/Plugins/Experimental/CommonUI/Source/CommonUI/Private/CommonButtonBase.cpp
+++ b/Engine/Plugins/Experimental/CommonUI/Source/CommonUI/Private/CommonButtonBase.cpp
@@ -1310,6 +1310,8 @@ void UCommonButtonBase::BuildStyles()
 		SelectedStyle.NormalPadding = ButtonPadding;
 		SelectedStyle.PressedPadding = ButtonPadding;
 
+		DisabledStyle = NormalStyle;
+
 		/**
 		 * Selected State Sound overrides
 		 * If there is no Selected state sound override, the Normal state's sound will be used.
