Removed Scale box in CommonActionWidget (this widget is often used inside SizeBox and having inside it ScaleBox creates undesired behavior). Kept the order of hold progress material behind the prompt icon.

diff --git a/Engine/Plugins/Experimental/CommonUI/Source/CommonUI/Private/CommonActionWidget.cpp b/Engine/Plugins/Experimental/CommonUI/Source/CommonUI/Private/CommonActionWidget.cpp
--- a/Engine/Plugins/Experimental/CommonUI/Source/CommonUI/Private/CommonActionWidget.cpp
+++ b/Engine/Plugins/Experimental/CommonUI/Source/CommonUI/Private/CommonActionWidget.cpp
@@ -59,35 +59,28 @@ TSharedRef<SWidget> UCommonActionWidget::RebuildWidget()
 		.HAlign(HAlign_Center)
 		.VAlign(VAlign_Center);
 
-	MyKeyBox->SetContent(
-		SNew(SScaleBox)
+	MyKeyBox->SetContent(	
+		SNew(SOverlay)
+		+ SOverlay::Slot()
 		.HAlign(HAlign_Center)
 		.VAlign(VAlign_Center)
-		.Stretch(IconStretch)
-		.UserSpecifiedScale(UserSpecifiedScale)
 		[
-			SNew(SOverlay)
-			+ SOverlay::Slot()
-			.HAlign(HAlign_Center)
-			.VAlign(VAlign_Center)
-			[
-				SAssignNew(MyIconRim, SImage)
-				.Image(&IconRimBrush)
-			]
-			+ SOverlay::Slot()
-			.HAlign(HAlign_Center)
-			.VAlign(VAlign_Center)
-			[
-				SAssignNew(MyProgressImage, SImage)
-				.Image(&ProgressMaterialBrush)
-			]
-			+ SOverlay::Slot()
-			.HAlign(HAlign_Center)
-			.VAlign(VAlign_Center)
-			[
-				SAssignNew(MyIcon, SImage)
-				.Image(&Icon)
-			]
+			SAssignNew(MyIconRim, SImage)
+			.Image(&IconRimBrush)
+		]
+		+ SOverlay::Slot()
+		.HAlign(HAlign_Center)
+		.VAlign(VAlign_Center)
+		[
+			SAssignNew(MyProgressImage, SImage)
+			.Image(&ProgressMaterialBrush)
+		]
+		+ SOverlay::Slot()
+		.HAlign(HAlign_Center)
+		.VAlign(VAlign_Center)
+		[
+			SAssignNew(MyIcon, SImage)
+			.Image(&Icon)
 		]);
 	
 	return MyKeyBox.ToSharedRef();
@@ -98,7 +91,6 @@ void UCommonActionWidget::ReleaseSlateResources(bool bReleaseChildren)
 	MyProgressImage.Reset();
 	MyIcon.Reset();
 	MyKeyBox.Reset();
-	MyScaleBox.Reset();
 	
 	ListenToInputMethodChanged(false);
 	Super::ReleaseSlateResources(bReleaseChildren);
diff --git a/Engine/Plugins/Experimental/CommonUI/Source/CommonUI/Public/CommonActionWidget.h b/Engine/Plugins/Experimental/CommonUI/Source/CommonUI/Public/CommonActionWidget.h
--- a/Engine/Plugins/Experimental/CommonUI/Source/CommonUI/Public/CommonActionWidget.h
+++ b/Engine/Plugins/Experimental/CommonUI/Source/CommonUI/Public/CommonActionWidget.h
@@ -4,7 +4,6 @@
 
 #include "CommonUITypes.h"
 #include "Input/UIActionBindingHandle.h"
-#include "Widgets/Layout/SScaleBox.h"
 
 #include "CommonActionWidget.generated.h"
 
@@ -69,12 +68,6 @@ public:
 
 	UPROPERTY(EditAnywhere, BlueprintReadOnly, Category = CommonActionWidget)
 	FSlateBrush IconRimBrush;
-	
-	UPROPERTY(EditAnywhere, BlueprintReadOnly, Category = CommonActionWidget)
-	TEnumAsByte<EStretch::Type> IconStretch = EStretch::UserSpecified;
-	
-	UPROPERTY(EditAnywhere, BlueprintReadOnly, Category = CommonActionWidget)
-	float UserSpecifiedScale = 1.0f;
 
 #if WITH_EDITOR
 	virtual const FText GetPaletteCategory() override;
@@ -120,8 +113,6 @@ protected:
 	
 	TSharedPtr<SBox> MyKeyBox;
 
-	TSharedPtr<SScaleBox> MyScaleBox;
-
 	TSharedPtr<SImage> MyIcon;
 
 	TSharedPtr<SImage> MyProgressImage;
