Adding Invisible SButton style for the CommonButtonBase and CVar UseTransparentButtonStyleAsDefault conditional.

diff --git a/Engine/Plugins/Experimental/CommonUI/Source/CommonUI/Private/CommonButtonBase.cpp b/Engine/Plugins/Experimental/CommonUI/Source/CommonUI/Private/CommonButtonBase.cpp
--- a/Engine/Plugins/Experimental/CommonUI/Source/CommonUI/Private/CommonButtonBase.cpp
+++ b/Engine/Plugins/Experimental/CommonUI/Source/CommonUI/Private/CommonButtonBase.cpp
@@ -24,6 +24,7 @@
 #include "CommonUIUtils.h"
 #include "Input/CommonUIInputTypes.h"
 #include "Sound/SoundBase.h"
+#include "Styling/UMGCoreStyle.h"
 
 #include UE_INLINE_GENERATED_CPP_BY_NAME(CommonButtonBase)
 
@@ -149,14 +150,27 @@ void UCommonButtonStyle::GetDisabledBrush(FSlateBrush& Brush) const
 }
 
 //////////////////////////////////////////////////////////////////////////
-// UCommonButtonBase
+// UCommonButtonInternalBase
 //////////////////////////////////////////////////////////////////////////
 
+static int32 bUseTransparentButtonStyleAsDefault = 0;
+static FAutoConsoleVariableRef CVarUseTransparentButtonStyleAsDefault(
+	TEXT("UseTransparentButtonStyleAsDefault"),
+	bUseTransparentButtonStyleAsDefault,
+	TEXT("If true, the default Button Style for the CommonButtonBase's SButton will be set to NoBorder, which has a transparent background and no padding"));
+
 UCommonButtonInternalBase::UCommonButtonInternalBase(const FObjectInitializer& Initializer)
 	: Super(Initializer)
 	, bButtonEnabled(true)
 	, bInteractionEnabled(true)
-{}
+{
+	if (bUseTransparentButtonStyleAsDefault)
+	{
+		// SButton will have a transparent background and have no padding if Button Style is set to None
+		static const FButtonStyle* TransparentButtonStyle = new FButtonStyle(FUMGCoreStyle::Get().GetWidgetStyle<FButtonStyle>("NoBorder"));
+		WidgetStyle = *TransparentButtonStyle;
+	}
+}
 
 void UCommonButtonInternalBase::SetButtonEnabled(bool bInIsButtonEnabled)
 {
