[Engine]
* AnimGraphNode_RemapCurvesFromMesh - Added include to prepare for coming refactoring of includes
* CommonUISettings - Added ctor/dtor to have implementation in cpp file to prepare for large change with forward declarations

diff --git a/Engine/Plugins/Experimental/Animation/CurveExpression/Source/Editor/Private/AnimGraphNode_RemapCurvesFromMesh.cpp b/Engine/Plugins/Experimental/Animation/CurveExpression/Source/Editor/Private/AnimGraphNode_RemapCurvesFromMesh.cpp
--- a/Engine/Plugins/Experimental/Animation/CurveExpression/Source/Editor/Private/AnimGraphNode_RemapCurvesFromMesh.cpp
+++ b/Engine/Plugins/Experimental/Animation/CurveExpression/Source/Editor/Private/AnimGraphNode_RemapCurvesFromMesh.cpp
@@ -3,6 +3,7 @@
 #include "AnimGraphNode_RemapCurvesFromMesh.h"
 
 #include "Animation/AnimAttributes.h"
+#include "Animation/AnimInstance.h"
 
 #include UE_INLINE_GENERATED_CPP_BY_NAME(AnimGraphNode_RemapCurvesFromMesh)
 
diff --git a/Engine/Plugins/Experimental/CommonUI/Source/CommonUI/Private/CommonUISettings.cpp b/Engine/Plugins/Experimental/CommonUI/Source/CommonUI/Private/CommonUISettings.cpp
--- a/Engine/Plugins/Experimental/CommonUI/Source/CommonUI/Private/CommonUISettings.cpp
+++ b/Engine/Plugins/Experimental/CommonUI/Source/CommonUI/Private/CommonUISettings.cpp
@@ -19,6 +19,12 @@ UCommonUISettings::UCommonUISettings(const FObjectInitializer& Initializer)
 	, bDefaultDataLoaded(false)
 {}
 
+UCommonUISettings::UCommonUISettings(FVTableHelper& Helper) : Super(Helper)
+{
+}
+
+UCommonUISettings::~UCommonUISettings() = default;
+
 void UCommonUISettings::LoadData()
 {
 	LoadEditorData();
diff --git a/Engine/Plugins/Experimental/CommonUI/Source/CommonUI/Public/CommonUISettings.h b/Engine/Plugins/Experimental/CommonUI/Source/CommonUI/Public/CommonUISettings.h
--- a/Engine/Plugins/Experimental/CommonUI/Source/CommonUI/Public/CommonUISettings.h
+++ b/Engine/Plugins/Experimental/CommonUI/Source/CommonUI/Public/CommonUISettings.h
@@ -25,7 +25,9 @@ class COMMONUI_API UCommonUISettings : public UObject
 	GENERATED_BODY()
 
 public:
-	UCommonUISettings(const FObjectInitializer& Initializer);
+	UCommonUISettings(const FObjectInitializer& Initializer = FObjectInitializer::Get());
+	UCommonUISettings(FVTableHelper& Helper);
+	~UCommonUISettings();
 
 	// Called to load CommonUISetting data, if bAutoLoadData if set to false then game code must call LoadData().
 	void LoadData();
