PR #9440: Give Child Classes some control of Mouse Position in CommonUI (Contributed by [REDACTED_PERSON])

diff --git a/Engine/Plugins/Experimental/CommonUI/Source/CommonUI/Public/Input/CommonUIActionRouterBase.h b/Engine/Plugins/Experimental/CommonUI/Source/CommonUI/Public/Input/CommonUIActionRouterBase.h
--- a/Engine/Plugins/Experimental/CommonUI/Source/CommonUI/Public/Input/CommonUIActionRouterBase.h
+++ b/Engine/Plugins/Experimental/CommonUI/Source/CommonUI/Public/Input/CommonUIActionRouterBase.h
@@ -125,6 +125,7 @@ protected:
 	virtual void SetActiveRoot(FActivatableTreeRootPtr NewActiveRoot);
 	void SetForceResetActiveRoot(bool bInForceResetActiveRoot);
 
+	virtual void ApplyUIInputConfig(const FUIInputConfig& NewConfig, bool bForceRefresh);
 	void UpdateLeafNodeAndConfig(FActivatableTreeRootPtr DesiredRoot, FActivatableTreeNodePtr DesiredLeafNode);
 
 	void RefreshActionDomainLeafNodeConfig();
@@ -146,8 +147,6 @@ private:
 	FActivatableTreeNodePtr FindOwningNode(const UWidget& Widget) const;
 	FActivatableTreeNodePtr FindNodeRecursive(const FActivatableTreeNodePtr& CurrentNode, const UCommonActivatableWidget& Widget) const;
 	FActivatableTreeNodePtr FindNodeRecursive(const FActivatableTreeNodePtr& CurrentNode, const TSharedPtr<SWidget>& Widget) const;
-
-	void ApplyUIInputConfig(const FUIInputConfig& NewConfig, bool bForceRefresh);
 	void SetActiveUICameraConfig(const FUICameraConfig& NewConfig);
 	
 	void HandleActivatableWidgetRebuilding(UCommonActivatableWidget& RebuildingWidget);
