[CommonUI]
* Fixed non-unity compile error

diff --git a/Engine/Plugins/Experimental/CommonUI/Source/CommonUI/Private/CommonVisibilityWidgetBase.cpp b/Engine/Plugins/Experimental/CommonUI/Source/CommonUI/Private/CommonVisibilityWidgetBase.cpp
--- a/Engine/Plugins/Experimental/CommonUI/Source/CommonUI/Private/CommonVisibilityWidgetBase.cpp
+++ b/Engine/Plugins/Experimental/CommonUI/Source/CommonUI/Private/CommonVisibilityWidgetBase.cpp
@@ -1,6 +1,7 @@
 // Copyright Epic Games, Inc. All Rights Reserved.
 
 #include "CommonVisibilityWidgetBase.h"
+#include "CommonInputSubsystem.h"
 #include "CommonUIPrivate.h"
 #include "CommonUISubsystemBase.h"
 
