Speculative fix crash on text style changed.

diff --git a/Engine/Plugins/Experimental/CommonUI/Source/CommonUI/Private/Input/CommonBoundActionButton.cpp b/Engine/Plugins/Experimental/CommonUI/Source/CommonUI/Private/Input/CommonBoundActionButton.cpp
--- a/Engine/Plugins/Experimental/CommonUI/Source/CommonUI/Private/Input/CommonBoundActionButton.cpp
+++ b/Engine/Plugins/Experimental/CommonUI/Source/CommonUI/Private/Input/CommonBoundActionButton.cpp
@@ -39,7 +39,11 @@ void UCommonBoundActionButton::NativeOnClicked()
 void UCommonBoundActionButton::NativeOnCurrentTextStyleChanged()
 {
 	Super::NativeOnCurrentTextStyleChanged();
-	Text_ActionName->SetStyle(GetCurrentTextStyleClass());
+
+	if (Text_ActionName)
+	{
+		Text_ActionName->SetStyle(GetCurrentTextStyleClass());
+	}
 }
 
 void UCommonBoundActionButton::UpdateInputActionWidget()
@@ -66,8 +70,11 @@ void UCommonBoundActionButton::UpdateInputActionWidget()
 			}
 		}
 
-		Text_ActionName->SetText(ActionDisplayName);
-
+		if (Text_ActionName)
+		{
+			Text_ActionName->SetText(ActionDisplayName);
+		}
+		
 		OnUpdateInputAction();
 	}
 }
