Fix an issue where commondatetimetextblock didn't broadcast compleation if you didn't use expiry text

diff --git a/Engine/Plugins/Experimental/CommonUI/Source/CommonUI/Private/CommonDateTimeTextBlock.cpp b/Engine/Plugins/Experimental/CommonUI/Source/CommonUI/Private/CommonDateTimeTextBlock.cpp
--- a/Engine/Plugins/Experimental/CommonUI/Source/CommonUI/Private/CommonDateTimeTextBlock.cpp
+++ b/Engine/Plugins/Experimental/CommonUI/Source/CommonUI/Private/CommonDateTimeTextBlock.cpp
@@ -96,10 +96,14 @@ void UCommonDateTimeTextBlock::UpdateUnderlyingText()
 			}
 		}
 
-		if (Remaining.GetTotalSeconds() < 1.0 && bUseCountdownCompletionText)
+		if (Remaining.GetTotalSeconds() < 1.0)
 		{
-			TextToSet = CountdownCompletionText;
 			OnTimeCountDownCompletion().Broadcast();
+			
+			if(bUseCountdownCompletionText)
+			{
+				TextToSet = CountdownCompletionText;
+			}
 		}
 		else
 		{
