Fix to CommonBoundActionBar to use extended pathname to correct warning appearing in the logs.

#rnx

diff --git a/Engine/Plugins/Experimental/CommonUI/Source/CommonUI/Public/Input/CommonBoundActionBar.h b/Engine/Plugins/Experimental/CommonUI/Source/CommonUI/Public/Input/CommonBoundActionBar.h
--- a/Engine/Plugins/Experimental/CommonUI/Source/CommonUI/Public/Input/CommonBoundActionBar.h
+++ b/Engine/Plugins/Experimental/CommonUI/Source/CommonUI/Public/Input/CommonBoundActionBar.h
@@ -37,7 +37,7 @@ private:
 	
 	void MonitorPlayerActions(const ULocalPlayer* NewPlayer);
 
-	UPROPERTY(EditAnywhere, Category = EntryLayout, meta=(MustImplement = "CommonBoundActionButtonInterface"))
+	UPROPERTY(EditAnywhere, Category = EntryLayout, meta=(MustImplement = "/Script/CommonUI.CommonBoundActionButtonInterface"))
 	TSubclassOf<UCommonButtonBase> ActionButtonClass;
 
 	UPROPERTY(EditAnywhere, Category = Display)
