Unity build fix: Missing SlateApplication include.

diff --git a/Engine/Plugins/Experimental/CommonUI/Source/CommonInput/Private/CommonInputSubsystem.cpp b/Engine/Plugins/Experimental/CommonUI/Source/CommonInput/Private/CommonInputSubsystem.cpp
--- a/Engine/Plugins/Experimental/CommonUI/Source/CommonInput/Private/CommonInputSubsystem.cpp
+++ b/Engine/Plugins/Experimental/CommonUI/Source/CommonInput/Private/CommonInputSubsystem.cpp
@@ -5,6 +5,7 @@
 #include "Misc/ConfigCacheIni.h"
 
 #include "Framework/Application/IInputProcessor.h"
+#include "Framework/Application/SlateApplication.h"
 #include "Framework/Application/SlateUser.h"
 #include "Engine/World.h"
 #include "Engine/GameViewportClient.h"
