Add mobile-specific properties to WBP_UIKit_CTAButton_Base

diff --git a/Engine/Plugins/Experimental/CommonUI/Source/CommonUI/Public/CommonTextBlock.h b/Engine/Plugins/Experimental/CommonUI/Source/CommonUI/Public/CommonTextBlock.h
--- a/Engine/Plugins/Experimental/CommonUI/Source/CommonUI/Public/CommonTextBlock.h
+++ b/Engine/Plugins/Experimental/CommonUI/Source/CommonUI/Public/CommonTextBlock.h
@@ -203,6 +203,10 @@ public:
 protected:
 	virtual TSharedRef<SWidget> RebuildWidget() override;
 	virtual void SynchronizeProperties() override;
+
+	/** Mobile font size multiplier. Activated by default on mobile. See CVar Mobile_PreviewFontSize */
+	UPROPERTY(EditAnywhere, BlueprintReadWrite, Category = "Mobile", meta = (ClampMin = "0.01", ClampMax = "5.0"))
+	float MobileFontSizeMultiplier = 1.0f;
 	
 #if WITH_EDITOR
 	virtual void OnCreationFromPalette() override;
@@ -235,10 +239,6 @@ private:
 	UPROPERTY(EditAnywhere, BlueprintReadOnly, Category = CommonText, meta = (AllowPrivateAccess = true))
 	bool bAutoCollapseWithEmptyText = false;
 
-	/** Mobile font size multiplier. Activated by default on mobile. See CVar Mobile_PreviewFontSize */
-	UPROPERTY(EditAnywhere, Category = "Mobile", meta = (ClampMin = "0.01", ClampMax = "5.0"))
-	float MobileFontSizeMultiplier = 1.0f;
-
 #if WITH_EDITORONLY_DATA
 	/** Used to track widgets that were created before changing the default style pointer to null */
 	UPROPERTY()
