Fixes to activatable tree nodes getting input incorrectly on action domains.
Common UI Action Router properly updates the leaf config when a widget in an action domain root deactivates.
#rnx

diff --git a/Engine/Plugins/Experimental/CommonUI/Source/CommonUI/Private/Input/CommonUIActionRouterBase.cpp b/Engine/Plugins/Experimental/CommonUI/Source/CommonUI/Private/Input/CommonUIActionRouterBase.cpp
--- a/Engine/Plugins/Experimental/CommonUI/Source/CommonUI/Private/Input/CommonUIActionRouterBase.cpp
+++ b/Engine/Plugins/Experimental/CommonUI/Source/CommonUI/Private/Input/CommonUIActionRouterBase.cpp
@@ -755,7 +755,7 @@ bool UCommonUIActionRouterBase::ProcessInputOnActionDomains(ECommonInputMode Act
 		{
 			if (RootNode->IsWidgetActivated())
 			{
-				bool bInputEventHandled = RootNode->ProcessActionDomainNormalInput(ActiveInputMode, Key, InputEvent);
+				bool bInputEventHandled = RootNode->ProcessNormalInput(ActiveInputMode, Key, InputEvent);
 				bInputEventHandledInDomain |= bInputEventHandled;
 				bDomainHadActiveRoots = true;
 
@@ -779,16 +779,16 @@ bool UCommonUIActionRouterBase::ProcessInputOnActionDomains(ECommonInputMode Act
 
 EProcessHoldActionResult UCommonUIActionRouterBase::ProcessHoldInputOnActionDomains(ECommonInputMode ActiveInputMode, FKey Key, EInputEvent InputEvent) const
 {
+	EProcessHoldActionResult HoldActionResult = EProcessHoldActionResult::Unhandled;
+
 	UCommonInputSubsystem& CommonInputSubsystem = GetInputSubsystem();
 	UCommonInputActionDomainTable* ActionDomainTable = CommonInputSubsystem.GetActionDomainTable();
 
 	if (!ActionDomainTable)
 	{
-		return EProcessHoldActionResult::Unhandled;
+		return HoldActionResult;
 	}
 
-	EProcessHoldActionResult OutHoldActionResult = EProcessHoldActionResult::Unhandled;
-
 	for (const UCommonInputActionDomain* ActionDomain : ActionDomainTable->ActionDomains)
 	{
 		const FActionDomainSortedRootList* SortedRootList = ActionDomainRootNodes.Find(ActionDomain);
@@ -803,13 +803,13 @@ EProcessHoldActionResult UCommonUIActionRouterBase::ProcessHoldInputOnActionDoma
 
 		for (const FActivatableTreeRootRef& RootNode : SortedRootList->RootList)
 		{
-			if (RootNode->IsWidgetActivated())
+			if (RootNode->IsReceivingInput() && HoldActionResult == EProcessHoldActionResult::Unhandled)
 			{
-				bool bInputEventHandled = RootNode->ProcessActionDomainHoldInput(ActiveInputMode, Key, InputEvent, OutHoldActionResult);
-				bInputEventHandledInDomain |= bInputEventHandled;
+				HoldActionResult = RootNode->ProcessHoldInput(ActiveInputMode, Key, InputEvent);
+				bInputEventHandledInDomain |= HoldActionResult == EProcessHoldActionResult::Handled;
 				bDomainHadActiveRoots = true;
 
-				if (ActionDomain->ShouldBreakInnerEventFlow(bInputEventHandled))
+				if (ActionDomain->ShouldBreakInnerEventFlow(HoldActionResult == EProcessHoldActionResult::Handled))
 				{
 					break;
 				}
@@ -823,7 +823,7 @@ EProcessHoldActionResult UCommonUIActionRouterBase::ProcessHoldInputOnActionDoma
 		}
 	}
 
-	return OutHoldActionResult;
+	return HoldActionResult;
 }
 
 void UCommonUIActionRouterBase::ProcessRebuiltWidgets()
@@ -1147,10 +1147,17 @@ void UCommonUIActionRouterBase::UpdateLeafNodeAndConfig(FActivatableTreeRootPtr
 {
 	if (DesiredRoot.IsValid())
 	{
-		// We're updating both the leaf node and it's config if we're the active root.
-		if (!DesiredRoot->UpdateLeafmostActiveNode(DesiredLeafNode, DesiredRoot == ActiveRootNode))
+		if (DesiredRoot == ActiveRootNode)
 		{
-			UE_LOG(LogUIActionRouter, Warning, TEXT("LeafmostActiveNode not updated."));
+			// We're updating both the leaf node and it's config if we're the active root.
+			if (!DesiredRoot->UpdateLeafmostActiveNode(DesiredLeafNode))
+			{
+				UE_LOG(LogUIActionRouter, Warning, TEXT("LeafmostActiveNode not updated."));
+			}
+		}
+		else
+		{
+			RefreshActionDomainLeafNodeConfig();
 		}
 	}
 }
diff --git a/Engine/Plugins/Experimental/CommonUI/Source/CommonUI/Private/Input/UIActionRouterTypes.cpp b/Engine/Plugins/Experimental/CommonUI/Source/CommonUI/Private/Input/UIActionRouterTypes.cpp
--- a/Engine/Plugins/Experimental/CommonUI/Source/CommonUI/Private/Input/UIActionRouterTypes.cpp
+++ b/Engine/Plugins/Experimental/CommonUI/Source/CommonUI/Private/Input/UIActionRouterTypes.cpp
@@ -712,21 +712,6 @@ EProcessHoldActionResult FActivatableTreeNode::ProcessHoldInput(ECommonInputMode
 	return EProcessHoldActionResult::Unhandled;
 }
 
-bool FActivatableTreeNode::ProcessActionDomainHoldInput(ECommonInputMode ActiveInputMode, FKey Key, EInputEvent InputEvent, EProcessHoldActionResult& OutHoldActionResult) const
-{
-	for (const FActivatableTreeNodeRef& ChildNode : Children)
-	{
-		EProcessHoldActionResult ChildResult = ChildNode->ProcessHoldInput(ActiveInputMode, Key, InputEvent);
-		if (ChildResult != EProcessHoldActionResult::Unhandled)
-		{
-			OutHoldActionResult = ChildResult;
-			return true;
-		}
-	}
-	OutHoldActionResult = FActionRouterBindingCollection::ProcessHoldInput(ActiveInputMode, Key, InputEvent);
-	return OutHoldActionResult != EProcessHoldActionResult::Unhandled;
-}
-
 bool FActivatableTreeNode::ProcessNormalInput(ECommonInputMode ActiveInputMode, FKey Key, EInputEvent InputEvent) const
 {
 	if (IsReceivingInput())
@@ -743,18 +728,6 @@ bool FActivatableTreeNode::ProcessNormalInput(ECommonInputMode ActiveInputMode,
 	return false;
 }
 
-bool FActivatableTreeNode::ProcessActionDomainNormalInput(ECommonInputMode ActiveInputMode, FKey Key, EInputEvent InputEvent) const
-{
-	for (const FActivatableTreeNodeRef& ChildNode : Children)
-	{
-		if (ChildNode->ProcessActionDomainNormalInput(ActiveInputMode, Key, InputEvent))
-		{
-			return true;
-		}
-	}
-	return FActionRouterBindingCollection::ProcessNormalInput(ActiveInputMode, Key, InputEvent);
-}
-
 bool FActivatableTreeNode::IsWidgetValid() const
 { 
 	return RepresentedWidget.IsValid();
diff --git a/Engine/Plugins/Experimental/CommonUI/Source/CommonUI/Private/Input/UIActionRouterTypes.h b/Engine/Plugins/Experimental/CommonUI/Source/CommonUI/Private/Input/UIActionRouterTypes.h
--- a/Engine/Plugins/Experimental/CommonUI/Source/CommonUI/Private/Input/UIActionRouterTypes.h
+++ b/Engine/Plugins/Experimental/CommonUI/Source/CommonUI/Private/Input/UIActionRouterTypes.h
@@ -90,9 +90,7 @@ public:
 	virtual ~FActivatableTreeNode();
 	
 	virtual EProcessHoldActionResult ProcessHoldInput(ECommonInputMode ActiveInputMode, FKey Key, EInputEvent InputEvent) const override;
-	virtual bool ProcessActionDomainHoldInput(ECommonInputMode ActiveInputMode, FKey Key, EInputEvent InputEvent, EProcessHoldActionResult& OutHoldActionResult) const;
 	virtual bool ProcessNormalInput(ECommonInputMode ActiveInputMode, FKey Key, EInputEvent InputEvent) const override;
-	virtual bool ProcessActionDomainNormalInput(ECommonInputMode ActiveInputMode, FKey Key, EInputEvent InputEvent) const;
 	virtual bool IsReceivingInput() const override { return bCanReceiveInput && IsWidgetActivated(); }
 
 	bool IsWidgetValid() const;
