Speculative fix call Super in virtual method to avoid stack overflow.

diff --git a/Engine/Plugins/Experimental/CommonUI/Source/CommonUI/Private/CommonAnimatedSwitcher.cpp b/Engine/Plugins/Experimental/CommonUI/Source/CommonUI/Private/CommonAnimatedSwitcher.cpp
--- a/Engine/Plugins/Experimental/CommonUI/Source/CommonUI/Private/CommonAnimatedSwitcher.cpp
+++ b/Engine/Plugins/Experimental/CommonUI/Source/CommonUI/Private/CommonAnimatedSwitcher.cpp
@@ -163,7 +163,7 @@ void UCommonAnimatedSwitcher::SetActiveWidgetIndex_Internal(int32 Index)
 	{
 		HandleOutgoingWidget();
 
-		SetActiveWidgetIndex(Index);
+		Super::SetActiveWidgetIndex(Index);
 
 		if (MyAnimatedSwitcher.IsValid())
 		{
