CommonUI - Don't allow CommonUI input preprocessor to eat input while intra frame debugging is happening for Blueprints.  This results in slate eatting input in the editor while stopped at a breakpoint effectively locking you out of the editor while debugging if you suspended input while launching say a dialog.

diff --git a/Engine/Plugins/Experimental/CommonUI/Source/CommonInput/Private/CommonInputSubsystem.cpp b/Engine/Plugins/Experimental/CommonUI/Source/CommonInput/Private/CommonInputSubsystem.cpp
--- a/Engine/Plugins/Experimental/CommonUI/Source/CommonInput/Private/CommonInputSubsystem.cpp
+++ b/Engine/Plugins/Experimental/CommonUI/Source/CommonInput/Private/CommonInputSubsystem.cpp
@@ -148,6 +148,17 @@ public:
 private:
 	bool IsRelevantInput(FSlateApplication& SlateApp, const FInputEvent& InputEvent, const ECommonInputType DesiredInputType)
 	{
+#if WITH_EDITOR
+		// If we're stopped at a breakpoint we need for this input preprocessor to just ignore all incoming input
+		// because we're now doing stuff outside the game loop in the editor and it needs to not block all that.
+		// This can happen if you suspend input while spawning a dialog and then hit another breakpoint and then
+		// try and use the editor, you can suddenly be unable to do anything.
+		if (GIntraFrameDebuggingGameThread)
+		{
+			return false;
+		}
+#endif
+		
 		if (SlateApp.IsActive() 
 			|| SlateApp.GetHandleDeviceInputWhenApplicationNotActive() 
 			|| (ICommonInputModule::GetSettings().GetAllowOutOfFocusDeviceInput() && DesiredInputType == ECommonInputType::Gamepad))
@@ -163,6 +174,7 @@ private:
 #endif
 			return ControllerId == InputEvent.GetUserIndex();
 		}
+		
 		return false;
 	}
 
