Make UCommonButtonBase play PressedSlateSound audio cue for buttons clicked via Triggering Input Action instead of Slate Press, similar to how UCommonButtonLegacy handles Triggering Input Actions.
#tests: PIE: Verified common button audio gets played for mouse interactions as well as gamepad Triggering Input Action interactions

diff --git a/Engine/Plugins/Experimental/CommonUI/Source/CommonUI/Private/CommonButtonBase.cpp b/Engine/Plugins/Experimental/CommonUI/Source/CommonUI/Private/CommonButtonBase.cpp
--- a/Engine/Plugins/Experimental/CommonUI/Source/CommonUI/Private/CommonButtonBase.cpp
+++ b/Engine/Plugins/Experimental/CommonUI/Source/CommonUI/Private/CommonButtonBase.cpp
@@ -504,7 +504,7 @@ void UCommonButtonBase::BindTriggeringInputActionToClick()
 
 	if (!TriggeringBindingHandle.IsValid())
 	{
-		FBindUIActionArgs BindArgs(TriggeringInputAction, false, FSimpleDelegate::CreateUObject(this, &UCommonButtonBase::HandleButtonClicked));
+		FBindUIActionArgs BindArgs(TriggeringInputAction, false, FSimpleDelegate::CreateUObject(this, &UCommonButtonBase::HandleTriggeringActionCommited));
 		BindArgs.OnHoldActionProgressed.BindUObject(this, &UCommonButtonBase::NativeOnActionProgress);
 		BindArgs.bIsPersistent = bIsPersistentBinding;
 
@@ -528,6 +528,11 @@ void UCommonButtonBase::UnbindTriggeringInputActionToClick()
 }
 
 void UCommonButtonBase::HandleTriggeringActionCommited(bool& bPassthrough)
+{
+	HandleTriggeringActionCommited();
+}
+
+void UCommonButtonBase::HandleTriggeringActionCommited()
 {
 	// Because this path doesn't go through SButton::Press(), the sound needs to be played from here.
 	FSlateApplication::Get().PlaySound(NormalStyle.PressedSlateSound);
diff --git a/Engine/Plugins/Experimental/CommonUI/Source/CommonUI/Public/CommonButtonBase.h b/Engine/Plugins/Experimental/CommonUI/Source/CommonUI/Public/CommonButtonBase.h
--- a/Engine/Plugins/Experimental/CommonUI/Source/CommonUI/Public/CommonButtonBase.h
+++ b/Engine/Plugins/Experimental/CommonUI/Source/CommonUI/Public/CommonButtonBase.h
@@ -499,6 +499,7 @@ protected:
 
 	UFUNCTION()
 	virtual void HandleTriggeringActionCommited(bool& bPassthrough);
+	virtual void HandleTriggeringActionCommited();
 
 	//@TODO: DarenC - API decision, consider removing
 	virtual void ExecuteTriggeredInput();
