Fix super / base class calls causing widgets to switch before they got a chance to animate, resulting in several regressions.

diff --git a/Engine/Plugins/Experimental/CommonUI/Source/CommonUI/Private/CommonAnimatedSwitcher.cpp b/Engine/Plugins/Experimental/CommonUI/Source/CommonUI/Private/CommonAnimatedSwitcher.cpp
--- a/Engine/Plugins/Experimental/CommonUI/Source/CommonUI/Private/CommonAnimatedSwitcher.cpp
+++ b/Engine/Plugins/Experimental/CommonUI/Source/CommonUI/Private/CommonAnimatedSwitcher.cpp
@@ -163,7 +163,14 @@ void UCommonAnimatedSwitcher::SetActiveWidgetIndex_Internal(int32 Index)
 	{
 		HandleOutgoingWidget();
 
-		Super::SetActiveWidgetIndex(Index);
+PRAGMA_DISABLE_DEPRECATION_WARNINGS
+		// For now we can't call setter since it calls MyWidgetSwitcher->SetActiveWidgetIndex(SafeIndex)
+		if (ActiveWidgetIndex != Index)
+		{
+			ActiveWidgetIndex = Index;
+			BroadcastFieldValueChanged(FFieldNotificationClassDescriptor::ActiveWidgetIndex);
+		}
+PRAGMA_ENABLE_DEPRECATION_WARNINGS
 
 		if (MyAnimatedSwitcher.IsValid())
 		{
