Adds Getter and Setter for MobileFontSizeMultiplier.

diff --git a/Engine/Plugins/Experimental/CommonUI/Source/CommonUI/Private/CommonTextBlock.cpp b/Engine/Plugins/Experimental/CommonUI/Source/CommonUI/Private/CommonTextBlock.cpp
--- a/Engine/Plugins/Experimental/CommonUI/Source/CommonUI/Private/CommonTextBlock.cpp
+++ b/Engine/Plugins/Experimental/CommonUI/Source/CommonUI/Private/CommonTextBlock.cpp
@@ -559,6 +559,24 @@ void UCommonTextBlock::SetScrollingEnabled(bool bInIsScrollingEnabled)
 	SynchronizeProperties();
 }
 
+float UCommonTextBlock::GetMobileFontSizeMultiplier() const
+{
+	return MobileFontSizeMultiplier;
+}
+
+void UCommonTextBlock::SetMobileFontSizeMultiplier(float InMobileFontSizeMultiplier)
+{
+	if (MobileFontSizeMultiplier != InMobileFontSizeMultiplier)
+	{
+		MobileFontSizeMultiplier = InMobileFontSizeMultiplier;
+
+		if (CommonUIUtils::ShouldDisplayMobileUISizes())
+		{
+			ApplyFontSizeMultiplier();
+		}
+	}
+}
+
 void UCommonTextBlock::SynchronizeProperties()
 {
 	UpdateFromStyle();
@@ -743,7 +761,7 @@ void UCommonTextBlock::ApplyFontSizeMultiplier() const
 	if (MyTextBlock.IsValid())
 	{
 		FSlateFontInfo FontInfo = GetFont();
-		FontInfo.Size *= MobileFontSizeMultiplier;
+		FontInfo.Size *= GetMobileFontSizeMultiplier();
 		MyTextBlock->SetFont(FontInfo);
 	}
 }
diff --git a/Engine/Plugins/Experimental/CommonUI/Source/CommonUI/Public/CommonTextBlock.h b/Engine/Plugins/Experimental/CommonUI/Source/CommonUI/Public/CommonTextBlock.h
--- a/Engine/Plugins/Experimental/CommonUI/Source/CommonUI/Public/CommonTextBlock.h
+++ b/Engine/Plugins/Experimental/CommonUI/Source/CommonUI/Public/CommonTextBlock.h
@@ -194,6 +194,13 @@ public:
 	UFUNCTION(BlueprintCallable, Category = "Common Text")
 	void SetMargin(const FMargin& InMargin);
 
+	UFUNCTION(BlueprintCallable, Category = "Common Text|Mobile")
+	float GetMobileFontSizeMultiplier() const;
+	
+	/** Sets the new value and then applies the FontSizeMultiplier */
+	UFUNCTION(BlueprintCallable, Category = "Common Text|Mobile")
+	void SetMobileFontSizeMultiplier(float InMobileFontSizeMultiplier);
+
 	UFUNCTION(BlueprintCallable, Category = "Common Text|Scroll Style")
 	void ResetScrollState();
 
@@ -205,7 +212,7 @@ protected:
 	virtual void SynchronizeProperties() override;
 
 	/** Mobile font size multiplier. Activated by default on mobile. See CVar Mobile_PreviewFontSize */
-	UPROPERTY(EditAnywhere, BlueprintReadWrite, Category = "Mobile", meta = (ClampMin = "0.01", ClampMax = "5.0"))
+	UPROPERTY(EditAnywhere, BlueprintReadWrite, Getter, Setter, Category = "Mobile", meta = (ClampMin = "0.01", ClampMax = "5.0"))
 	float MobileFontSizeMultiplier = 1.0f;
 	
 #if WITH_EDITOR
