PR #9864: Added CVar to enable the platform cursor when using gamepad (Contributed by [REDACTED_PERSON])

diff --git a/Engine/Plugins/Experimental/CommonUI/Source/CommonInput/Private/CommonInputSubsystem.cpp b/Engine/Plugins/Experimental/CommonUI/Source/CommonInput/Private/CommonInputSubsystem.cpp
--- a/Engine/Plugins/Experimental/CommonUI/Source/CommonInput/Private/CommonInputSubsystem.cpp
+++ b/Engine/Plugins/Experimental/CommonUI/Source/CommonInput/Private/CommonInputSubsystem.cpp
@@ -301,6 +301,14 @@ static FAutoConsoleVariableRef CVarInputKeysVisible
 	ECVF_Default
 );
 
+bool bEnableGamepadPlatformCursor = false;
+static const FAutoConsoleVariableRef CVarInputEnableGamepadPlatformCursor
+(
+	TEXT("CommonInput.EnableGamepadPlatformCursor"),
+	bEnableGamepadPlatformCursor,
+	TEXT("Should the cursor be allowed to be used during gamepad input")
+);
+
 UCommonInputSubsystem* UCommonInputSubsystem::Get(const ULocalPlayer* LocalPlayer)
 {
 	return LocalPlayer ? LocalPlayer->GetSubsystem<UCommonInputSubsystem>() : nullptr;
@@ -524,7 +532,7 @@ void UCommonInputSubsystem::SetCurrentInputType(ECommonInputType NewInputType)
 					UE_LOG(LogCommonInput, Log, TEXT("UCommonInputSubsystem::SetCurrentInputType(): Using Gamepad"));
 					if (bCursorUser)
 					{
-						SlateApplication.UsePlatformCursorForCursorUser(false);
+						SlateApplication.UsePlatformCursorForCursorUser(bEnableGamepadPlatformCursor);
 					}
 					break;
 				case ECommonInputType::Touch:
