Fixed analog navigation input not working properly in CommonUI AnalogSlider widget

diff --git a/Engine/Plugins/Experimental/CommonUI/Source/CommonUI/Private/SAnalogSlider.cpp b/Engine/Plugins/Experimental/CommonUI/Source/CommonUI/Private/SAnalogSlider.cpp
--- a/Engine/Plugins/Experimental/CommonUI/Source/CommonUI/Private/SAnalogSlider.cpp
+++ b/Engine/Plugins/Experimental/CommonUI/Source/CommonUI/Private/SAnalogSlider.cpp
@@ -62,6 +62,10 @@ FReply SAnalogSlider::OnAnalogValueChanged(const FGeometry& MyGeometry, const FA
 	{
 		Reply = FReply::Handled();
 	}
+	else
+	{
+		Reply = SSlider::OnAnalogValueChanged(MyGeometry, InAnalogInputEvent);
+	}
 
 	if (bIsUsingGamepad)
 	{
