PR #9895: Fix CommonWidgetSwitcher not working (Contributed by [REDACTED_PERSON])

diff --git a/Engine/Plugins/Experimental/CommonUI/Source/CommonUI/Private/CommonAnimatedSwitcher.cpp b/Engine/Plugins/Experimental/CommonUI/Source/CommonUI/Private/CommonAnimatedSwitcher.cpp
--- a/Engine/Plugins/Experimental/CommonUI/Source/CommonUI/Private/CommonAnimatedSwitcher.cpp
+++ b/Engine/Plugins/Experimental/CommonUI/Source/CommonUI/Private/CommonAnimatedSwitcher.cpp
@@ -170,20 +170,20 @@ PRAGMA_DISABLE_DEPRECATION_WARNINGS
 			ActiveWidgetIndex = Index;
 			BroadcastFieldValueChanged(FFieldNotificationClassDescriptor::ActiveWidgetIndex);
 		}
-PRAGMA_ENABLE_DEPRECATION_WARNINGS
 
 		if (MyAnimatedSwitcher.IsValid())
 		{
 			// Ensure the index is clamped to a valid range.
-			int32 SafeIndex = FMath::Clamp(GetActiveWidgetIndex(), 0, FMath::Max(0, Slots.Num() - 1));
+			int32 SafeIndex = FMath::Clamp(ActiveWidgetIndex, 0, FMath::Max(0, Slots.Num() - 1));
 			MyAnimatedSwitcher->TransitionToIndex(SafeIndex, bInstantTransition);
 		}
 
 		//When we're setting up, and the index goes 0->0, MyAnimatedSwitcher won't fire its ActiveIndexChanged Event.
 		if (!bSetOnce)
 		{
-			HandleSlateActiveIndexChanged(GetActiveWidgetIndex());
+			HandleSlateActiveIndexChanged(ActiveWidgetIndex);
 		}
+PRAGMA_ENABLE_DEPRECATION_WARNINGS
 
 		bSetOnce = true;
 	}
