Only updating move/look input locks when the value differs between new/old input configs.

diff --git a/Engine/Plugins/Experimental/CommonUI/Source/CommonUI/Private/Input/CommonUIActionRouterBase.cpp b/Engine/Plugins/Experimental/CommonUI/Source/CommonUI/Private/Input/CommonUIActionRouterBase.cpp
--- a/Engine/Plugins/Experimental/CommonUI/Source/CommonUI/Private/Input/CommonUIActionRouterBase.cpp
+++ b/Engine/Plugins/Experimental/CommonUI/Source/CommonUI/Private/Input/CommonUIActionRouterBase.cpp
@@ -1333,6 +1333,7 @@ void UCommonUIActionRouterBase::ApplyUIInputConfig(const FUIInputConfig& NewConf
 
 		const ECommonInputMode PreviousInputMode = GetActiveInputMode();
 
+		TOptional<FUIInputConfig> OldConfig = ActiveInputConfig;
 		ActiveInputConfig = NewConfig;
 
 		ULocalPlayer& LocalPlayer = *GetLocalPlayerChecked();
@@ -1344,8 +1345,15 @@ void UCommonUIActionRouterBase::ApplyUIInputConfig(const FUIInputConfig& NewConf
 			{
 				if (APlayerController* PC = LocalPlayer.GetPlayerController(GetWorld()))
 				{
-					PC->SetIgnoreMoveInput(NewConfig.bIgnoreMoveInput);
-					PC->SetIgnoreLookInput(NewConfig.bIgnoreLookInput);
+					if (!OldConfig.IsSet() || OldConfig.GetValue().bIgnoreMoveInput != NewConfig.bIgnoreMoveInput)
+					{
+						PC->SetIgnoreMoveInput(NewConfig.bIgnoreMoveInput);						
+					}
+					
+					if (!OldConfig.IsSet() || OldConfig.GetValue().bIgnoreLookInput != NewConfig.bIgnoreLookInput)
+					{
+						PC->SetIgnoreLookInput(NewConfig.bIgnoreLookInput);						
+					}
 
 					EMouseCaptureMode PrevCaptureMode = GameViewportClient->GetMouseCaptureMode();
 					const bool bWasPermanentlyCaptured = PrevCaptureMode == EMouseCaptureMode::CapturePermanently || PrevCaptureMode == EMouseCaptureMode::CapturePermanently_IncludingInitialMouseDown;
