Fix for Common UI Actions not visible on the navigation bar inside paused menus in cooked games

diff --git a/Engine/Plugins/Experimental/CommonUI/Source/CommonUI/Private/Input/CommonBoundActionBar.cpp b/Engine/Plugins/Experimental/CommonUI/Source/CommonUI/Private/Input/CommonBoundActionBar.cpp
--- a/Engine/Plugins/Experimental/CommonUI/Source/CommonUI/Private/Input/CommonBoundActionBar.cpp
+++ b/Engine/Plugins/Experimental/CommonUI/Source/CommonUI/Private/Input/CommonBoundActionBar.cpp
@@ -35,6 +35,29 @@ void UCommonBoundActionBar::SetDisplayOwningPlayerActionsOnly(bool bShouldOnlyDi
 	}
 }
 
+void UCommonBoundActionBar::Tick(float DeltaTime)
+{
+	if (bIsRefreshQueued)
+	{
+		HandleDeferredDisplayUpdate();
+	}
+}
+
+ETickableTickType UCommonBoundActionBar::GetTickableTickType() const
+{
+	return ETickableTickType::Always;
+}
+
+TStatId UCommonBoundActionBar::GetStatId() const
+{
+	RETURN_QUICK_DECLARE_CYCLE_STAT(UCommonBoundActionBar, STATGROUP_Tickables);
+}
+
+bool UCommonBoundActionBar::IsTickableWhenPaused() const
+{
+	return true;
+}
+
 void UCommonBoundActionBar::OnWidgetRebuilt()
 {
 	Super::OnWidgetRebuilt();
@@ -100,13 +123,9 @@ void UCommonBoundActionBar::ValidateCompiledDefaults(IWidgetCompilerLog& Compile
 
 void UCommonBoundActionBar::HandleBoundActionsUpdated(bool bFromOwningPlayer)
 {
-	if (!bIsRefreshQueued && (bFromOwningPlayer || !bDisplayOwningPlayerActionsOnly))
+	if (bFromOwningPlayer || !bDisplayOwningPlayerActionsOnly)
 	{
 		bIsRefreshQueued = true;
-
-		UGameInstance* GameInstance = GetGameInstance();
-		check(GameInstance);
-		GameInstance->GetTimerManager().SetTimerForNextTick(FTimerDelegate::CreateUObject(this, &UCommonBoundActionBar::HandleDeferredDisplayUpdate));
 	}
 }
 
diff --git a/Engine/Plugins/Experimental/CommonUI/Source/CommonUI/Public/Input/CommonBoundActionBar.h b/Engine/Plugins/Experimental/CommonUI/Source/CommonUI/Public/Input/CommonBoundActionBar.h
--- a/Engine/Plugins/Experimental/CommonUI/Source/CommonUI/Public/Input/CommonBoundActionBar.h
+++ b/Engine/Plugins/Experimental/CommonUI/Source/CommonUI/Public/Input/CommonBoundActionBar.h
@@ -4,6 +4,7 @@
 
 #include "CommonButtonBase.h"
 #include "Components/DynamicEntryBoxBase.h"
+#include "Tickable.h"
 #include "CommonBoundActionBar.generated.h"
 
 class IWidgetCompilerLog;
@@ -13,7 +14,7 @@ class IConsoleVariable;
  * A box populated with current actions available per CommonUI's Input Handler.
  */
 UCLASS(Blueprintable, ClassGroup = UI, meta = (Category = "Common UI"))
-class COMMONUI_API UCommonBoundActionBar : public UDynamicEntryBoxBase
+class COMMONUI_API UCommonBoundActionBar : public UDynamicEntryBoxBase, public FTickableGameObject
 {
 	GENERATED_BODY()
 
@@ -21,6 +22,13 @@ public:
 	UFUNCTION(BlueprintCallable, Category = CommonBoundActionBar)
 	void SetDisplayOwningPlayerActionsOnly(bool bShouldOnlyDisplayOwningPlayerActions);
 
+	//~ FTickableGameObject Begin
+	virtual void Tick(float DeltaTime) override;
+	virtual ETickableTickType GetTickableTickType() const override;
+	virtual TStatId GetStatId() const override;
+	virtual bool IsTickableWhenPaused() const override;
+	//~ FTickableGameObject End
+
 protected:
 	virtual void OnWidgetRebuilt() override;
 	virtual void SynchronizeProperties() override;
