[CommonUI]
* Ran IWYU on cpp files to remove lots of includes not needed.

diff --git a/Engine/Plugins/Experimental/CommonUI/Source/CommonInput/Private/CommonInputBaseTypes.cpp b/Engine/Plugins/Experimental/CommonUI/Source/CommonInput/Private/CommonInputBaseTypes.cpp
--- a/Engine/Plugins/Experimental/CommonUI/Source/CommonInput/Private/CommonInputBaseTypes.cpp
+++ b/Engine/Plugins/Experimental/CommonUI/Source/CommonInput/Private/CommonInputBaseTypes.cpp
@@ -2,13 +2,13 @@
 
 #include "CommonInputBaseTypes.h"
 #include "Engine/UserInterfaceSettings.h"
-#include "Engine/StreamableManager.h"
-#include "Engine/AssetManager.h"
-#include "CommonInputSettings.h"
+#include "CommonInputTypeEnum.h"
 #include "ICommonInputModule.h"
+#include "Engine/PlatformSettings.h"
 #include "Misc/DataDrivenPlatformInfoRegistry.h"
+#include "InputCoreTypes.h"
 #include "UObject/ObjectSaveContext.h"
-#include "CommonInputBaseTypes.h"
+#include "Styling/SlateBrush.h"
 
 #include UE_INLINE_GENERATED_CPP_BY_NAME(CommonInputBaseTypes)
 
diff --git a/Engine/Plugins/Experimental/CommonUI/Source/CommonInput/Private/CommonInputModule.cpp b/Engine/Plugins/Experimental/CommonUI/Source/CommonInput/Private/CommonInputModule.cpp
--- a/Engine/Plugins/Experimental/CommonUI/Source/CommonInput/Private/CommonInputModule.cpp
+++ b/Engine/Plugins/Experimental/CommonUI/Source/CommonInput/Private/CommonInputModule.cpp
@@ -1,12 +1,11 @@
 // Copyright Epic Games, Inc. All Rights Reserved.
 
 #include "CommonInputPrivate.h"
+#include "CommonInputSettings.h"
 #include "ICommonInputModule.h"
-#include "Engine/AssetManager.h"
 #include "UObject/Package.h"
 #if WITH_EDITOR
 #include "ISettingsModule.h"
-#include "ISettingsSection.h"
 #endif // WITH_EDITOR
 
 /**
@@ -89,4 +88,4 @@ UCommonInputSettings* FCommonInputModule::GetSettingsInstance() const
 
 IMPLEMENT_MODULE(FCommonInputModule, CommonInput);
 
-DEFINE_LOG_CATEGORY(LogCommonInput);
\ No newline at end of file
+DEFINE_LOG_CATEGORY(LogCommonInput);
diff --git a/Engine/Plugins/Experimental/CommonUI/Source/CommonInput/Private/CommonInputSettings.cpp b/Engine/Plugins/Experimental/CommonUI/Source/CommonInput/Private/CommonInputSettings.cpp
--- a/Engine/Plugins/Experimental/CommonUI/Source/CommonInput/Private/CommonInputSettings.cpp
+++ b/Engine/Plugins/Experimental/CommonUI/Source/CommonInput/Private/CommonInputSettings.cpp
@@ -1,17 +1,15 @@
 // Copyright Epic Games, Inc. All Rights Reserved.
 
 #include "CommonInputSettings.h"
-#include "CommonInputPrivate.h"
-#include "Engine/UserInterfaceSettings.h"
-#include "Engine/StreamableManager.h"
-#include "Engine/AssetManager.h"
-#include "Misc/DataDrivenPlatformInfoRegistry.h"
-#include "CommonInputActionDomain.h"
-#include "CommonInputBaseTypes.h"
-#include "Engine/PlatformSettings.h"
+#include "Engine/DataTable.h"
+#include "Engine/PlatformSettingsManager.h"
 
 #include UE_INLINE_GENERATED_CPP_BY_NAME(CommonInputSettings)
 
+#if !WITH_EDITOR
+#include "CommonInputPrivate.h"
+#endif
+
 UCommonInputSettings::UCommonInputSettings(const FObjectInitializer& Initializer)
 	: Super(Initializer)
 	, bInputDataLoaded(false)
diff --git a/Engine/Plugins/Experimental/CommonUI/Source/CommonInput/Private/CommonInputSubsystem.cpp b/Engine/Plugins/Experimental/CommonUI/Source/CommonInput/Private/CommonInputSubsystem.cpp
--- a/Engine/Plugins/Experimental/CommonUI/Source/CommonInput/Private/CommonInputSubsystem.cpp
+++ b/Engine/Plugins/Experimental/CommonUI/Source/CommonInput/Private/CommonInputSubsystem.cpp
@@ -2,23 +2,20 @@
 
 #include "CommonInputSubsystem.h"
 #include "CommonInputPrivate.h"
+#include "CommonInputTypeEnum.h"
 #include "Framework/Application/SlateApplication.h"
-#include "Misc/ConfigCacheIni.h"
 
+#include "Engine/PlatformSettingsManager.h"
 #include "Framework/Application/IInputProcessor.h"
 #include "Framework/Application/SlateApplication.h"
 #include "Framework/Application/SlateUser.h"
 #include "Engine/World.h"
 #include "Engine/GameViewportClient.h"
 #include "Widgets/SViewport.h"
-#include "HAL/IConsoleManager.h"
 #include "CommonInputSettings.h"
-#include "Containers/Ticker.h"
-#include "GenericPlatform/GenericPlatformTime.h"
 #include "ICommonInputModule.h"
 #include "Engine/LocalPlayer.h"
 #include "Engine/Engine.h"
-#include "Stats/Stats.h"
 
 #include UE_INLINE_GENERATED_CPP_BY_NAME(CommonInputSubsystem)
 
diff --git a/Engine/Plugins/Experimental/CommonUI/Source/CommonUI/Private/AnalogSlider.cpp b/Engine/Plugins/Experimental/CommonUI/Source/CommonUI/Private/AnalogSlider.cpp
--- a/Engine/Plugins/Experimental/CommonUI/Source/CommonUI/Private/AnalogSlider.cpp
+++ b/Engine/Plugins/Experimental/CommonUI/Source/CommonUI/Private/AnalogSlider.cpp
@@ -1,8 +1,9 @@
 // Copyright Epic Games, Inc. All Rights Reserved.
 
 #include "AnalogSlider.h"
-#include "CommonUITypes.h"
-#include "Widgets/DeclarativeSyntaxSupport.h"
+#include "CommonInputSubsystem.h"
+#include "CommonInputTypeEnum.h"
+#include "SAnalogSlider.h"
 
 #include UE_INLINE_GENERATED_CPP_BY_NAME(AnalogSlider)
 
diff --git a/Engine/Plugins/Experimental/CommonUI/Source/CommonUI/Private/CommonActionWidget.cpp b/Engine/Plugins/Experimental/CommonUI/Source/CommonUI/Private/CommonActionWidget.cpp
--- a/Engine/Plugins/Experimental/CommonUI/Source/CommonUI/Private/CommonActionWidget.cpp
+++ b/Engine/Plugins/Experimental/CommonUI/Source/CommonUI/Private/CommonActionWidget.cpp
@@ -1,9 +1,8 @@
 // Copyright Epic Games, Inc. All Rights Reserved.
 
 #include "CommonActionWidget.h"
-#include "CommonUIPrivate.h"
-#include "CommonUISubsystemBase.h"
 #include "CommonInputSubsystem.h"
+#include "CommonInputTypeEnum.h"
 #include "CommonWidgetPaletteCategories.h"
 #include "Widgets/Layout/SBox.h"
 #include "Widgets/Images/SImage.h"
@@ -11,7 +10,6 @@
 #include "CommonUITypes.h"
 #include "Materials/MaterialInstanceDynamic.h"
 #include "Input/UIActionBinding.h"
-#include "Input/UIActionRouterTypes.h"
 
 #include UE_INLINE_GENERATED_CPP_BY_NAME(CommonActionWidget)
 
diff --git a/Engine/Plugins/Experimental/CommonUI/Source/CommonUI/Private/CommonActivatableWidget.cpp b/Engine/Plugins/Experimental/CommonUI/Source/CommonUI/Private/CommonActivatableWidget.cpp
--- a/Engine/Plugins/Experimental/CommonUI/Source/CommonUI/Private/CommonActivatableWidget.cpp
+++ b/Engine/Plugins/Experimental/CommonUI/Source/CommonUI/Private/CommonActivatableWidget.cpp
@@ -1,9 +1,9 @@
 // Copyright Epic Games, Inc. All Rights Reserved.
 
 #include "CommonActivatableWidget.h"
+#include "CommonInputSettings.h"
 #include "CommonUIPrivate.h"
 #include "Engine/GameInstance.h"
-#include "Input/CommonUIActionRouterBase.h"
 #include "Input/CommonUIInputTypes.h"
 #include "Input/UIActionBinding.h"
 #include "Input/UIActionRouterTypes.h"
diff --git a/Engine/Plugins/Experimental/CommonUI/Source/CommonUI/Private/CommonAnimatedSwitcher.cpp b/Engine/Plugins/Experimental/CommonUI/Source/CommonUI/Private/CommonAnimatedSwitcher.cpp
--- a/Engine/Plugins/Experimental/CommonUI/Source/CommonUI/Private/CommonAnimatedSwitcher.cpp
+++ b/Engine/Plugins/Experimental/CommonUI/Source/CommonUI/Private/CommonAnimatedSwitcher.cpp
@@ -3,9 +3,7 @@
 #include "CommonAnimatedSwitcher.h"
 
 #include "CommonWidgetPaletteCategories.h"
-#include "Components/PanelSlot.h"
 #include "Components/WidgetSwitcherSlot.h"
-#include "Templates/UnrealTemplate.h"
 #include "Widgets/Layout/SSpacer.h"
 #include "Widgets/SOverlay.h"
 
diff --git a/Engine/Plugins/Experimental/CommonUI/Source/CommonUI/Private/CommonBorder.cpp b/Engine/Plugins/Experimental/CommonUI/Source/CommonUI/Private/CommonBorder.cpp
--- a/Engine/Plugins/Experimental/CommonUI/Source/CommonUI/Private/CommonBorder.cpp
+++ b/Engine/Plugins/Experimental/CommonUI/Source/CommonUI/Private/CommonBorder.cpp
@@ -1,11 +1,11 @@
 // Copyright Epic Games, Inc. All Rights Reserved.
 
 #include "CommonBorder.h"
-#include "CommonUIPrivate.h"
 #include "CommonUIEditorSettings.h"
 #include "CommonWidgetPaletteCategories.h"
 
 #include "Framework/Application/SlateApplication.h"
+#include "ICommonUIModule.h"
 #include "Misc/CoreDelegates.h"
 #include "Widgets/Layout/SBorder.h"
 #include "Widgets/SViewport.h"
diff --git a/Engine/Plugins/Experimental/CommonUI/Source/CommonUI/Private/CommonButtonBase.cpp b/Engine/Plugins/Experimental/CommonUI/Source/CommonUI/Private/CommonButtonBase.cpp
--- a/Engine/Plugins/Experimental/CommonUI/Source/CommonUI/Private/CommonButtonBase.cpp
+++ b/Engine/Plugins/Experimental/CommonUI/Source/CommonUI/Private/CommonButtonBase.cpp
@@ -1,27 +1,23 @@
 // Copyright Epic Games, Inc. All Rights Reserved.
 
 #include "CommonButtonBase.h"
-#include "CommonUIPrivate.h"
 
-#include "CommonTextBlock.h"
 #include "CommonActionWidget.h"
+#include "CommonInputTypeEnum.h"
 #include "CommonUISubsystemBase.h"
 #include "CommonInputSubsystem.h"
-#include "CommonUISettings.h"
 #include "CommonUIEditorSettings.h"
 #include "CommonWidgetPaletteCategories.h"
 #include "Components/ButtonSlot.h"
 #include "Blueprint/WidgetTree.h"
 #include "CommonButtonTypes.h"
-#include "Widgets/Input/SButton.h"
+#include "ICommonUIModule.h"
 #include "Widgets/Layout/SBox.h"
 #include "Framework/Application/SlateApplication.h"
 #include "Materials/MaterialInstanceDynamic.h"
-#include "Materials/MaterialInterface.h"
 #include "Engine/UserInterfaceSettings.h"
 #include "ICommonInputModule.h"
 #include "CommonInputSettings.h"
-#include "CommonUIUtils.h"
 #include "Input/CommonUIInputTypes.h"
 #include "Sound/SoundBase.h"
 #include "Styling/UMGCoreStyle.h"
diff --git a/Engine/Plugins/Experimental/CommonUI/Source/CommonUI/Private/CommonCustomNavigation.cpp b/Engine/Plugins/Experimental/CommonUI/Source/CommonUI/Private/CommonCustomNavigation.cpp
--- a/Engine/Plugins/Experimental/CommonUI/Source/CommonUI/Private/CommonCustomNavigation.cpp
+++ b/Engine/Plugins/Experimental/CommonUI/Source/CommonUI/Private/CommonCustomNavigation.cpp
@@ -1,15 +1,12 @@
 // Copyright Epic Games, Inc. All Rights Reserved.
 
 #include "CommonCustomNavigation.h"
-#include "CommonUIPrivate.h"
-#include "CommonUISettings.h"
 #include "CommonWidgetPaletteCategories.h"
 
 #include "Types/NavigationMetaData.h"
 #include "Widgets/Layout/SBorder.h"
 #include "Components/BorderSlot.h"
 
-#include "Framework/Application/SlateApplication.h"
 
 #include UE_INLINE_GENERATED_CPP_BY_NAME(CommonCustomNavigation)
 
diff --git a/Engine/Plugins/Experimental/CommonUI/Source/CommonUI/Private/CommonDateTimeTextBlock.cpp b/Engine/Plugins/Experimental/CommonUI/Source/CommonUI/Private/CommonDateTimeTextBlock.cpp
--- a/Engine/Plugins/Experimental/CommonUI/Source/CommonUI/Private/CommonDateTimeTextBlock.cpp
+++ b/Engine/Plugins/Experimental/CommonUI/Source/CommonUI/Private/CommonDateTimeTextBlock.cpp
@@ -1,11 +1,9 @@
 // Copyright Epic Games, Inc. All Rights Reserved.
 
 #include "CommonDateTimeTextBlock.h"
-#include "CommonUIPrivate.h"
 #include "CommonWidgetPaletteCategories.h"
+#include "Engine/World.h"
 #include "TimerManager.h"
-#include "Widgets/Text/STextBlock.h"
-#include "Framework/Application/SlateApplication.h"
 
 #include UE_INLINE_GENERATED_CPP_BY_NAME(CommonDateTimeTextBlock)
 
diff --git a/Engine/Plugins/Experimental/CommonUI/Source/CommonUI/Private/CommonGameViewportClient.cpp b/Engine/Plugins/Experimental/CommonUI/Source/CommonUI/Private/CommonGameViewportClient.cpp
--- a/Engine/Plugins/Experimental/CommonUI/Source/CommonUI/Private/CommonGameViewportClient.cpp
+++ b/Engine/Plugins/Experimental/CommonUI/Source/CommonUI/Private/CommonGameViewportClient.cpp
@@ -2,21 +2,14 @@
 
 #include "CommonGameViewportClient.h"
 #include "Engine/Console.h"
-#include "Engine/Engine.h"
 #include "Engine/GameInstance.h"
 #include "Engine/LocalPlayer.h"
-#include "GameFramework/GameUserSettings.h"
-#include "GameFramework/PlayerController.h"
-#include "GenericPlatform/GenericPlatformInputDeviceMapper.h"
-#include "UnrealClient.h"
+#include "InputKeyEventArgs.h"
 
 #if WITH_EDITOR
-#include "Editor.h"
-#include "EditorViewportClient.h"
 #endif // WITH_EDITOR
 
 #include "Input/CommonUIActionRouterBase.h"
-#include "Framework/Application/SlateUser.h"
 
 #include UE_INLINE_GENERATED_CPP_BY_NAME(CommonGameViewportClient)
 
diff --git a/Engine/Plugins/Experimental/CommonUI/Source/CommonUI/Private/CommonHardwareVisibilityBorder.cpp b/Engine/Plugins/Experimental/CommonUI/Source/CommonUI/Private/CommonHardwareVisibilityBorder.cpp
--- a/Engine/Plugins/Experimental/CommonUI/Source/CommonUI/Private/CommonHardwareVisibilityBorder.cpp
+++ b/Engine/Plugins/Experimental/CommonUI/Source/CommonUI/Private/CommonHardwareVisibilityBorder.cpp
@@ -2,10 +2,7 @@
 
 #include "CommonHardwareVisibilityBorder.h"
 #include "CommonUIPrivate.h"
-#include "CommonUISubsystemBase.h"
 
-#include "CommonInputBaseTypes.h"
-#include "Misc/DataDrivenPlatformInfoRegistry.h"
 #include "CommonUIVisibilitySubsystem.h"
 #include "CommonUISettings.h"
 #include "ICommonUIModule.h"
diff --git a/Engine/Plugins/Experimental/CommonUI/Source/CommonUI/Private/CommonLazyImage.cpp b/Engine/Plugins/Experimental/CommonUI/Source/CommonUI/Private/CommonLazyImage.cpp
--- a/Engine/Plugins/Experimental/CommonUI/Source/CommonUI/Private/CommonLazyImage.cpp
+++ b/Engine/Plugins/Experimental/CommonUI/Source/CommonUI/Private/CommonLazyImage.cpp
@@ -4,6 +4,7 @@
 #include "CommonWidgetPaletteCategories.h"
 
 #include "Engine/Texture2D.h"
+#include "ICommonUIModule.h"
 #include "Materials/MaterialInstanceDynamic.h"
 #include "CommonUISettings.h"
 #include "Engine/Texture2DDynamic.h"
diff --git a/Engine/Plugins/Experimental/CommonUI/Source/CommonUI/Private/CommonLazyWidget.cpp b/Engine/Plugins/Experimental/CommonUI/Source/CommonUI/Private/CommonLazyWidget.cpp
--- a/Engine/Plugins/Experimental/CommonUI/Source/CommonUI/Private/CommonLazyWidget.cpp
+++ b/Engine/Plugins/Experimental/CommonUI/Source/CommonUI/Private/CommonLazyWidget.cpp
@@ -2,8 +2,8 @@
 
 #include "CommonLazyWidget.h"
 #include "CommonActivatableWidget.h"
-#include "CommonUISettings.h"
 #include "CommonWidgetPaletteCategories.h"
+#include "Engine/AssetManager.h"
 
 #include UE_INLINE_GENERATED_CPP_BY_NAME(CommonLazyWidget)
 
diff --git a/Engine/Plugins/Experimental/CommonUI/Source/CommonUI/Private/CommonListView.cpp b/Engine/Plugins/Experimental/CommonUI/Source/CommonUI/Private/CommonListView.cpp
--- a/Engine/Plugins/Experimental/CommonUI/Source/CommonUI/Private/CommonListView.cpp
+++ b/Engine/Plugins/Experimental/CommonUI/Source/CommonUI/Private/CommonListView.cpp
@@ -3,7 +3,6 @@
 #include "CommonListView.h"
 #include "CommonWidgetPaletteCategories.h"
 #include "SCommonButtonTableRow.h"
-#include "CommonUIPrivate.h"
 
 #include UE_INLINE_GENERATED_CPP_BY_NAME(CommonListView)
 
diff --git a/Engine/Plugins/Experimental/CommonUI/Source/CommonUI/Private/CommonLoadGuard.cpp b/Engine/Plugins/Experimental/CommonUI/Source/CommonUI/Private/CommonLoadGuard.cpp
--- a/Engine/Plugins/Experimental/CommonUI/Source/CommonUI/Private/CommonLoadGuard.cpp
+++ b/Engine/Plugins/Experimental/CommonUI/Source/CommonUI/Private/CommonLoadGuard.cpp
@@ -1,14 +1,14 @@
 // Copyright Epic Games, Inc. All Rights Reserved.
 
 #include "CommonLoadGuard.h"
-#include "CommonUIPrivate.h"
 #include "Components/SizeBoxSlot.h"
-#include "CommonTextBlock.h"
 #include "CommonWidgetPaletteCategories.h"
-#include "UObject/ConstructorHelpers.h"
 #include "CommonUIObjectVersion.h"
+#include "Engine/AssetManager.h"
 #include "Widgets/SOverlay.h"
+#include "Engine/StreamableManager.h"
 #include "Widgets/Text/STextBlock.h"
+#include "Widgets/Layout/SBorder.h"
 #include "Widgets/SBoxPanel.h"
 #include "Widgets/Images/SImage.h"
 #include "Widgets/Layout/SScaleBox.h"
diff --git a/Engine/Plugins/Experimental/CommonUI/Source/CommonUI/Private/CommonNumericTextBlock.cpp b/Engine/Plugins/Experimental/CommonUI/Source/CommonUI/Private/CommonNumericTextBlock.cpp
--- a/Engine/Plugins/Experimental/CommonUI/Source/CommonUI/Private/CommonNumericTextBlock.cpp
+++ b/Engine/Plugins/Experimental/CommonUI/Source/CommonUI/Private/CommonNumericTextBlock.cpp
@@ -1,9 +1,10 @@
 // Copyright Epic Games, Inc. All Rights Reserved.
 
 #include "CommonNumericTextBlock.h"
-#include "CommonUIPrivate.h"
 #include "CommonWidgetPaletteCategories.h"
+#include "Engine/World.h"
 #include "TimerManager.h"
+#include "UObject/UnrealType.h"
 #include "Widgets/Text/STextBlock.h"
 #include "Framework/Application/SlateApplication.h"
 #include "Math/UnitConversion.h"
diff --git a/Engine/Plugins/Experimental/CommonUI/Source/CommonUI/Private/CommonRichTextBlock.cpp b/Engine/Plugins/Experimental/CommonUI/Source/CommonUI/Private/CommonRichTextBlock.cpp
--- a/Engine/Plugins/Experimental/CommonUI/Source/CommonUI/Private/CommonRichTextBlock.cpp
+++ b/Engine/Plugins/Experimental/CommonUI/Source/CommonUI/Private/CommonRichTextBlock.cpp
@@ -1,29 +1,23 @@
 // Copyright Epic Games, Inc. All Rights Reserved.
 
 #include "CommonRichTextBlock.h"
+#include "ICommonUIModule.h"
 #include "Widgets/Text/SRichTextBlock.h"
+#include "Rendering/SlateRenderer.h"
 #include "Widgets/Text/STextBlock.h"
 #include "Widgets/Images/SImage.h"
 #include "Widgets/Layout/SBox.h"
-#include "Widgets/Layout/SScaleBox.h"
 
-#include "Framework/Application/SlateApplication.h"
-#include "Framework/Text/SlateTextLayout.h"
-#include "Framework/Text/ITextDecorator.h"
-#include "Framework/Text/SlateWidgetRun.h"
 #include "Framework/Text/SlateTextRun.h"
 
 #include "Engine/AssetManager.h"
 #include "Engine/StreamableManager.h"
 #include "Misc/DefaultValueHelper.h"
 #include "Fonts/FontMeasure.h"
-#include "Styling/SlateBrush.h"
 #include "CommonUIUtils.h"
 #include "CommonWidgetPaletteCategories.h"
 #include "CommonUIEditorSettings.h"
 #include "CommonUIRichTextData.h"
-#include "Types/ReflectionMetadata.h"
-#include "Framework/Text/IRichTextMarkupWriter.h"
 #include "Framework/Text/RichTextMarkupProcessing.h"
 #include "Framework/Application/SlateApplication.h"
 
diff --git a/Engine/Plugins/Experimental/CommonUI/Source/CommonUI/Private/CommonRotator.cpp b/Engine/Plugins/Experimental/CommonUI/Source/CommonUI/Private/CommonRotator.cpp
--- a/Engine/Plugins/Experimental/CommonUI/Source/CommonUI/Private/CommonRotator.cpp
+++ b/Engine/Plugins/Experimental/CommonUI/Source/CommonUI/Private/CommonRotator.cpp
@@ -1,6 +1,7 @@
 // Copyright Epic Games, Inc. All Rights Reserved.
 
 #include "CommonRotator.h"
+#include "CommonTextBlock.h"
 #include "CommonUIPrivate.h"
 
 #include UE_INLINE_GENERATED_CPP_BY_NAME(CommonRotator)
diff --git a/Engine/Plugins/Experimental/CommonUI/Source/CommonUI/Private/CommonTabListWidgetBase.cpp b/Engine/Plugins/Experimental/CommonUI/Source/CommonUI/Private/CommonTabListWidgetBase.cpp
--- a/Engine/Plugins/Experimental/CommonUI/Source/CommonUI/Private/CommonTabListWidgetBase.cpp
+++ b/Engine/Plugins/Experimental/CommonUI/Source/CommonUI/Private/CommonTabListWidgetBase.cpp
@@ -1,15 +1,9 @@
 // Copyright Epic Games, Inc. All Rights Reserved.
 
 #include "CommonTabListWidgetBase.h"
-#include "CommonUIPrivate.h"
 
-#include "Blueprint/WidgetLayoutLibrary.h"
 #include "CommonAnimatedSwitcher.h"
-#include "CommonUISubsystemBase.h"
-#include "CommonUIUtils.h"
-#include "Containers/Ticker.h"
 #include "Groups/CommonButtonGroupBase.h"
-#include "ICommonUIModule.h"
 #include "Input/CommonUIInputTypes.h"
 
 #include UE_INLINE_GENERATED_CPP_BY_NAME(CommonTabListWidgetBase)
diff --git a/Engine/Plugins/Experimental/CommonUI/Source/CommonUI/Private/CommonTextBlock.cpp b/Engine/Plugins/Experimental/CommonUI/Source/CommonUI/Private/CommonTextBlock.cpp
--- a/Engine/Plugins/Experimental/CommonUI/Source/CommonUI/Private/CommonTextBlock.cpp
+++ b/Engine/Plugins/Experimental/CommonUI/Source/CommonUI/Private/CommonTextBlock.cpp
@@ -1,14 +1,13 @@
 // Copyright Epic Games, Inc. All Rights Reserved.
 
 #include "CommonTextBlock.h"
-#include "CommonUIPrivate.h"
-#include "Containers/Ticker.h"
 #include "CommonUIEditorSettings.h"
 #include "CommonUIUtils.h"
 #include "CommonWidgetPaletteCategories.h"
+#include "ICommonUIModule.h"
 #include "Widgets/Text/STextBlock.h"
+#include "Layout/ArrangedChildren.h"
 #include "Layout/LayoutUtils.h"
-#include "Types/ReflectionMetadata.h"
 #include "StyleSheet/CommonStyleSheet.h"
 
 #include UE_INLINE_GENERATED_CPP_BY_NAME(CommonTextBlock)
diff --git a/Engine/Plugins/Experimental/CommonUI/Source/CommonUI/Private/CommonTileView.cpp b/Engine/Plugins/Experimental/CommonUI/Source/CommonUI/Private/CommonTileView.cpp
--- a/Engine/Plugins/Experimental/CommonUI/Source/CommonUI/Private/CommonTileView.cpp
+++ b/Engine/Plugins/Experimental/CommonUI/Source/CommonUI/Private/CommonTileView.cpp
@@ -1,7 +1,7 @@
 // Copyright Epic Games, Inc. All Rights Reserved.
 
 #include "CommonTileView.h"
-#include "CommonUIPrivate.h"
+#include "Components/ListView.h"
 #include "SCommonButtonTableRow.h"
 #include "SCommonTileView.h"
 
diff --git a/Engine/Plugins/Experimental/CommonUI/Source/CommonUI/Private/CommonTreeView.cpp b/Engine/Plugins/Experimental/CommonUI/Source/CommonUI/Private/CommonTreeView.cpp
--- a/Engine/Plugins/Experimental/CommonUI/Source/CommonUI/Private/CommonTreeView.cpp
+++ b/Engine/Plugins/Experimental/CommonUI/Source/CommonUI/Private/CommonTreeView.cpp
@@ -1,7 +1,7 @@
 // Copyright Epic Games, Inc. All Rights Reserved.
 
 #include "CommonTreeView.h"
-#include "CommonUIPrivate.h"
+#include "Components/ListView.h"
 #include "SCommonButtonTableRow.h"
 
 #include UE_INLINE_GENERATED_CPP_BY_NAME(CommonTreeView)
diff --git a/Engine/Plugins/Experimental/CommonUI/Source/CommonUI/Private/CommonUIEditorSettings.cpp b/Engine/Plugins/Experimental/CommonUI/Source/CommonUI/Private/CommonUIEditorSettings.cpp
--- a/Engine/Plugins/Experimental/CommonUI/Source/CommonUI/Private/CommonUIEditorSettings.cpp
+++ b/Engine/Plugins/Experimental/CommonUI/Source/CommonUI/Private/CommonUIEditorSettings.cpp
@@ -2,15 +2,7 @@
 
 
 #include "CommonUIEditorSettings.h"
-#include "CommonTextBlock.h"
-#include "CommonButtonBase.h"
-#include "CommonBorder.h"
-
-#include "Materials/MaterialInterface.h"
-#include "Materials/Material.h"
-#include "Engine/UserInterfaceSettings.h"
-#include "Engine/StreamableManager.h"
-#include "Engine/AssetManager.h"
+
 
 #include UE_INLINE_GENERATED_CPP_BY_NAME(CommonUIEditorSettings)
 
diff --git a/Engine/Plugins/Experimental/CommonUI/Source/CommonUI/Private/CommonUIModule.cpp b/Engine/Plugins/Experimental/CommonUI/Source/CommonUI/Private/CommonUIModule.cpp
--- a/Engine/Plugins/Experimental/CommonUI/Source/CommonUI/Private/CommonUIModule.cpp
+++ b/Engine/Plugins/Experimental/CommonUI/Source/CommonUI/Private/CommonUIModule.cpp
@@ -5,7 +5,6 @@
 #include "CommonUISettings.h"
 #include "CommonUIEditorSettings.h"
 #include "Engine/AssetManager.h"
-#include "Misc/CoreDelegates.h"
 
 #include "Input/CommonUIInputSettings.h"
 #include "UObject/Package.h"
@@ -152,4 +151,4 @@ UCommonUIEditorSettings* FCommonUIModule::GetEditorSettingsInstance() const
 
 IMPLEMENT_MODULE(FCommonUIModule, CommonUI);
 
-DEFINE_LOG_CATEGORY(LogCommonUI);
\ No newline at end of file
+DEFINE_LOG_CATEGORY(LogCommonUI);
diff --git a/Engine/Plugins/Experimental/CommonUI/Source/CommonUI/Private/CommonUIRichTextData.cpp b/Engine/Plugins/Experimental/CommonUI/Source/CommonUI/Private/CommonUIRichTextData.cpp
--- a/Engine/Plugins/Experimental/CommonUI/Source/CommonUI/Private/CommonUIRichTextData.cpp
+++ b/Engine/Plugins/Experimental/CommonUI/Source/CommonUI/Private/CommonUIRichTextData.cpp
@@ -3,6 +3,7 @@
 #include "CommonUIRichTextData.h"
 
 #include "CommonUISettings.h"
+#include "ICommonUIModule.h"
 
 #include UE_INLINE_GENERATED_CPP_BY_NAME(CommonUIRichTextData)
 
diff --git a/Engine/Plugins/Experimental/CommonUI/Source/CommonUI/Private/CommonUISettings.cpp b/Engine/Plugins/Experimental/CommonUI/Source/CommonUI/Private/CommonUISettings.cpp
--- a/Engine/Plugins/Experimental/CommonUI/Source/CommonUI/Private/CommonUISettings.cpp
+++ b/Engine/Plugins/Experimental/CommonUI/Source/CommonUI/Private/CommonUISettings.cpp
@@ -2,11 +2,8 @@
 
 #include "CommonUISettings.h"
 
+#include "Engine/PlatformSettingsManager.h"
 #include "Materials/MaterialInterface.h"
-#include "Materials/Material.h"
-#include "UObject/Class.h"
-#include "UObject/ReflectedTypeAccessors.h"
-#include "GameplayTagContainer.h"
 #include "Misc/ConfigCacheIni.h"
 
 #include UE_INLINE_GENERATED_CPP_BY_NAME(CommonUISettings)
diff --git a/Engine/Plugins/Experimental/CommonUI/Source/CommonUI/Private/CommonUISubsystemBase.cpp b/Engine/Plugins/Experimental/CommonUI/Source/CommonUI/Private/CommonUISubsystemBase.cpp
--- a/Engine/Plugins/Experimental/CommonUI/Source/CommonUI/Private/CommonUISubsystemBase.cpp
+++ b/Engine/Plugins/Experimental/CommonUI/Source/CommonUI/Private/CommonUISubsystemBase.cpp
@@ -2,19 +2,15 @@
 
 #include "CommonUISubsystemBase.h"
 
-#include "Subsystems/SubsystemCollection.h"
 
-#include "Blueprint/UserWidget.h"
-#include "CommonUIPrivate.h"
+#include "CommonInputBaseTypes.h"
 #include "CommonInputSubsystem.h"
-#include "ICommonInputModule.h"
-#include "CommonInputSettings.h"
-#include "CommonActivatableWidget.h"
+#include "CommonInputTypeEnum.h"
+#include "CommonUITypes.h"
 #include "IAnalyticsProviderET.h"
-#include "AnalyticsEventAttribute.h"
-#include "UObject/UObjectIterator.h"
-#include "Input/CommonUIActionRouterBase.h"
-#include "Engine/LocalPlayer.h"
+#include "Components/Widget.h"
+#include "Engine/GameInstance.h"
+#include "Styling/StyleDefaults.h"
 
 #include UE_INLINE_GENERATED_CPP_BY_NAME(CommonUISubsystemBase)
 
diff --git a/Engine/Plugins/Experimental/CommonUI/Source/CommonUI/Private/CommonUITypes.cpp b/Engine/Plugins/Experimental/CommonUI/Source/CommonUI/Private/CommonUITypes.cpp
--- a/Engine/Plugins/Experimental/CommonUI/Source/CommonUI/Private/CommonUITypes.cpp
+++ b/Engine/Plugins/Experimental/CommonUI/Source/CommonUI/Private/CommonUITypes.cpp
@@ -1,10 +1,13 @@
 // Copyright Epic Games, Inc. All Rights Reserved.
  
 #include "CommonUITypes.h"
+#include "CommonInputBaseTypes.h"
+#include "CommonInputTypeEnum.h"
 #include "CommonUIPrivate.h"
-#include "ICommonInputModule.h"
-#include "CommonInputSettings.h"
+#include "CommonInputSubsystem.h"
 #include "HAL/PlatformInput.h"
+#include "Styling/SlateTypes.h"
+#include "Styling/StyleDefaults.h"
 
 #include UE_INLINE_GENERATED_CPP_BY_NAME(CommonUITypes)
 
diff --git a/Engine/Plugins/Experimental/CommonUI/Source/CommonUI/Private/CommonUIUtils.cpp b/Engine/Plugins/Experimental/CommonUI/Source/CommonUI/Private/CommonUIUtils.cpp
--- a/Engine/Plugins/Experimental/CommonUI/Source/CommonUI/Private/CommonUIUtils.cpp
+++ b/Engine/Plugins/Experimental/CommonUI/Source/CommonUI/Private/CommonUIUtils.cpp
@@ -1,10 +1,7 @@
 // Copyright Epic Games, Inc. All Rights Reserved.
  
 #include "CommonUIUtils.h"
-#include "ICommonUIModule.h"
-#include "CommonUIPrivate.h"
-#include "CommonUIEditorSettings.h"
-#include "Blueprint/WidgetTree.h"
+#include "Components/PanelWidget.h"
 #include "Editor/WidgetCompilerLog.h"
 
 #define LOCTEXT_NAMESPACE "CommonUIUtils"
@@ -85,4 +82,4 @@ namespace CommonUIUtils
 #endif
 }
 
-#undef LOCTEXT_NAMESPACE 
\ No newline at end of file
+#undef LOCTEXT_NAMESPACE 
diff --git a/Engine/Plugins/Experimental/CommonUI/Source/CommonUI/Private/CommonUIVisibilitySubsystem.cpp b/Engine/Plugins/Experimental/CommonUI/Source/CommonUI/Private/CommonUIVisibilitySubsystem.cpp
--- a/Engine/Plugins/Experimental/CommonUI/Source/CommonUI/Private/CommonUIVisibilitySubsystem.cpp
+++ b/Engine/Plugins/Experimental/CommonUI/Source/CommonUI/Private/CommonUIVisibilitySubsystem.cpp
@@ -2,13 +2,10 @@
 
 #include "CommonUIVisibilitySubsystem.h"
 
-#include "NativeGameplayTags.h"
-#include "Engine/PlatformSettings.h"
-#include "CommonInputBaseTypes.h"
+#include "CommonInputTypeEnum.h"
 #include "Engine/LocalPlayer.h"
 #include "Engine/GameInstance.h"
 #include "CommonInputSubsystem.h"
-#include "UObject/UObjectHash.h"
 #include "UObject/UObjectIterator.h"
 #include "ICommonUIModule.h"
 #include "CommonUISettings.h"
diff --git a/Engine/Plugins/Experimental/CommonUI/Source/CommonUI/Private/CommonUserWidget.cpp b/Engine/Plugins/Experimental/CommonUI/Source/CommonUI/Private/CommonUserWidget.cpp
--- a/Engine/Plugins/Experimental/CommonUI/Source/CommonUI/Private/CommonUserWidget.cpp
+++ b/Engine/Plugins/Experimental/CommonUI/Source/CommonUI/Private/CommonUserWidget.cpp
@@ -1,7 +1,6 @@
 // Copyright Epic Games, Inc. All Rights Reserved.
 
 #include "CommonUserWidget.h"
-#include "CommonUIPrivate.h"
 
 #include "Engine/GameInstance.h"
 #include "CommonInputSubsystem.h"
diff --git a/Engine/Plugins/Experimental/CommonUI/Source/CommonUI/Private/CommonVideoPlayer.cpp b/Engine/Plugins/Experimental/CommonUI/Source/CommonUI/Private/CommonVideoPlayer.cpp
--- a/Engine/Plugins/Experimental/CommonUI/Source/CommonUI/Private/CommonVideoPlayer.cpp
+++ b/Engine/Plugins/Experimental/CommonUI/Source/CommonUI/Private/CommonVideoPlayer.cpp
@@ -9,7 +9,6 @@
 #include "Materials/Material.h"
 #include "Materials/MaterialInstanceDynamic.h"
 #include "MediaSoundComponent.h"
-#include "ShaderPipelineCache.h"
 #include "IMediaEventSink.h"
 #include "Widgets/Images/SImage.h"
 
diff --git a/Engine/Plugins/Experimental/CommonUI/Source/CommonUI/Private/CommonVisibilitySwitcher.cpp b/Engine/Plugins/Experimental/CommonUI/Source/CommonUI/Private/CommonVisibilitySwitcher.cpp
--- a/Engine/Plugins/Experimental/CommonUI/Source/CommonUI/Private/CommonVisibilitySwitcher.cpp
+++ b/Engine/Plugins/Experimental/CommonUI/Source/CommonUI/Private/CommonVisibilitySwitcher.cpp
@@ -6,6 +6,7 @@
 #include "CommonUIPrivate.h"
 #include "CommonVisibilitySwitcherSlot.h"
 #include "CommonWidgetPaletteCategories.h"
+#include "Containers/StringConv.h"
 #include "Widgets/Layout/SBox.h"
 
 #include UE_INLINE_GENERATED_CPP_BY_NAME(CommonVisibilitySwitcher)
diff --git a/Engine/Plugins/Experimental/CommonUI/Source/CommonUI/Private/CommonVisibilityWidgetBase.cpp b/Engine/Plugins/Experimental/CommonUI/Source/CommonUI/Private/CommonVisibilityWidgetBase.cpp
--- a/Engine/Plugins/Experimental/CommonUI/Source/CommonUI/Private/CommonVisibilityWidgetBase.cpp
+++ b/Engine/Plugins/Experimental/CommonUI/Source/CommonUI/Private/CommonVisibilityWidgetBase.cpp
@@ -2,12 +2,9 @@
 
 #include "CommonVisibilityWidgetBase.h"
 #include "CommonInputSubsystem.h"
+#include "CommonInputTypeEnum.h"
 #include "CommonUIPrivate.h"
-#include "CommonUISubsystemBase.h"
 
-#include "Components/BorderSlot.h"
-#include "Widgets/Layout/SBox.h"
-#include "Styling/CoreStyle.h"
 #include "CommonInputBaseTypes.h"
 
 #include UE_INLINE_GENERATED_CPP_BY_NAME(CommonVisibilityWidgetBase)
diff --git a/Engine/Plugins/Experimental/CommonUI/Source/CommonUI/Private/CommonVisualAttachment.cpp b/Engine/Plugins/Experimental/CommonUI/Source/CommonUI/Private/CommonVisualAttachment.cpp
--- a/Engine/Plugins/Experimental/CommonUI/Source/CommonUI/Private/CommonVisualAttachment.cpp
+++ b/Engine/Plugins/Experimental/CommonUI/Source/CommonUI/Private/CommonVisualAttachment.cpp
@@ -2,8 +2,6 @@
 
 #include "CommonVisualAttachment.h"
 #include "SVisualAttachmentBox.h"
-#include "Widgets/DeclarativeSyntaxSupport.h"
-#include "Framework/Application/SlateApplication.h"
 #include "Components/SizeBoxSlot.h"
 
 #include UE_INLINE_GENERATED_CPP_BY_NAME(CommonVisualAttachment)
diff --git a/Engine/Plugins/Experimental/CommonUI/Source/CommonUI/Private/CommonWidgetCarousel.cpp b/Engine/Plugins/Experimental/CommonUI/Source/CommonUI/Private/CommonWidgetCarousel.cpp
--- a/Engine/Plugins/Experimental/CommonUI/Source/CommonUI/Private/CommonWidgetCarousel.cpp
+++ b/Engine/Plugins/Experimental/CommonUI/Source/CommonUI/Private/CommonWidgetCarousel.cpp
@@ -1,9 +1,7 @@
 // Copyright Epic Games, Inc. All Rights Reserved.
 
 #include "CommonWidgetCarousel.h"
-#include "CommonUIPrivate.h"
 #include "CommonWidgetPaletteCategories.h"
-#include "Containers/Ticker.h"
 
 #include UE_INLINE_GENERATED_CPP_BY_NAME(CommonWidgetCarousel)
 
diff --git a/Engine/Plugins/Experimental/CommonUI/Source/CommonUI/Private/CommonWidgetCarouselNavBar.cpp b/Engine/Plugins/Experimental/CommonUI/Source/CommonUI/Private/CommonWidgetCarouselNavBar.cpp
--- a/Engine/Plugins/Experimental/CommonUI/Source/CommonUI/Private/CommonWidgetCarouselNavBar.cpp
+++ b/Engine/Plugins/Experimental/CommonUI/Source/CommonUI/Private/CommonWidgetCarouselNavBar.cpp
@@ -1,11 +1,9 @@
 // Copyright Epic Games, Inc. All Rights Reserved.
 
 #include "CommonWidgetCarouselNavBar.h"
-#include "CommonUIPrivate.h"
 #include "CommonWidgetPaletteCategories.h"
 #include "CommonWidgetCarousel.h"
 #include "Groups/CommonButtonGroupBase.h"
-#include "Widgets/SBoxPanel.h"
 
 #include UE_INLINE_GENERATED_CPP_BY_NAME(CommonWidgetCarouselNavBar)
 
diff --git a/Engine/Plugins/Experimental/CommonUI/Source/CommonUI/Private/Groups/CommonButtonGroupBase.cpp b/Engine/Plugins/Experimental/CommonUI/Source/CommonUI/Private/Groups/CommonButtonGroupBase.cpp
--- a/Engine/Plugins/Experimental/CommonUI/Source/CommonUI/Private/Groups/CommonButtonGroupBase.cpp
+++ b/Engine/Plugins/Experimental/CommonUI/Source/CommonUI/Private/Groups/CommonButtonGroupBase.cpp
@@ -1,7 +1,8 @@
 // Copyright Epic Games, Inc. All Rights Reserved.
 
 #include "Groups/CommonButtonGroupBase.h"
-#include "CommonUIPrivate.h"
+#include <memory>
+#include <tuple>
 
 #include UE_INLINE_GENERATED_CPP_BY_NAME(CommonButtonGroupBase)
 
diff --git a/Engine/Plugins/Experimental/CommonUI/Source/CommonUI/Private/Groups/CommonWidgetGroupBase.cpp b/Engine/Plugins/Experimental/CommonUI/Source/CommonUI/Private/Groups/CommonWidgetGroupBase.cpp
--- a/Engine/Plugins/Experimental/CommonUI/Source/CommonUI/Private/Groups/CommonWidgetGroupBase.cpp
+++ b/Engine/Plugins/Experimental/CommonUI/Source/CommonUI/Private/Groups/CommonWidgetGroupBase.cpp
@@ -1,7 +1,6 @@
 // Copyright Epic Games, Inc. All Rights Reserved.
 
 #include "Groups/CommonWidgetGroupBase.h"
-#include "CommonUIPrivate.h"
 
 #include UE_INLINE_GENERATED_CPP_BY_NAME(CommonWidgetGroupBase)
 
diff --git a/Engine/Plugins/Experimental/CommonUI/Source/CommonUI/Private/Input/CommonAnalogCursor.cpp b/Engine/Plugins/Experimental/CommonUI/Source/CommonUI/Private/Input/CommonAnalogCursor.cpp
--- a/Engine/Plugins/Experimental/CommonUI/Source/CommonUI/Private/Input/CommonAnalogCursor.cpp
+++ b/Engine/Plugins/Experimental/CommonUI/Source/CommonUI/Private/Input/CommonAnalogCursor.cpp
@@ -1,27 +1,20 @@
 // Copyright Epic Games, Inc. All Rights Reserved.
 
 #include "Input/CommonAnalogCursor.h"
+#include "CommonInputTypeEnum.h"
 #include "CommonUIPrivate.h"
-#include "Slate/SObjectWidget.h"
-#include "Blueprint/UserWidget.h"
 #include "Input/CommonUIActionRouterBase.h"
 #include "CommonInputSubsystem.h"
-#include "Framework/Application/SlateApplication.h"
-#include "Engine/LocalPlayer.h"
 #include "Framework/Application/SlateUser.h"
-#include "Framework/Views/TableViewMetadata.h"
 #include "Input/CommonUIInputSettings.h"
 #include "Engine/Console.h"
-#include "CommonInputBaseTypes.h"
 #include "Widgets/SViewport.h"
 #include "Engine/GameViewportClient.h"
 
 #include "Components/ListView.h"
 #include "Components/ScrollBar.h"
 #include "Components/ScrollBox.h"
-#include "Engine/Engine.h"
 #include "Slate/SGameLayerManager.h"
-#include "Types/ReflectionMetadata.h"
 #include "UnrealClient.h"
 
 #define LOCTEXT_NAMESPACE "CommonAnalogCursor"
diff --git a/Engine/Plugins/Experimental/CommonUI/Source/CommonUI/Private/Input/CommonBoundActionBar.cpp b/Engine/Plugins/Experimental/CommonUI/Source/CommonUI/Private/Input/CommonBoundActionBar.cpp
--- a/Engine/Plugins/Experimental/CommonUI/Source/CommonUI/Private/Input/CommonBoundActionBar.cpp
+++ b/Engine/Plugins/Experimental/CommonUI/Source/CommonUI/Private/Input/CommonBoundActionBar.cpp
@@ -2,13 +2,14 @@
 
 #include "Input/CommonBoundActionBar.h"
 #include "CommonInputSubsystem.h"
+#include "CommonInputTypeEnum.h"
 #include "Editor/WidgetCompilerLog.h"
+#include "CommonUITypes.h"
 #include "Engine/GameInstance.h"
 #include "Engine/GameViewportClient.h"
-#include "Engine/LocalPlayer.h"
 #include "Input/CommonBoundActionButtonInterface.h"
 #include "Input/CommonUIActionRouterBase.h"
-#include "Input/UIActionRouterTypes.h"
+#include "Input/UIActionBinding.h"
 #include "TimerManager.h"
 
 #include UE_INLINE_GENERATED_CPP_BY_NAME(CommonBoundActionBar)
diff --git a/Engine/Plugins/Experimental/CommonUI/Source/CommonUI/Private/Input/CommonBoundActionButton.cpp b/Engine/Plugins/Experimental/CommonUI/Source/CommonUI/Private/Input/CommonBoundActionButton.cpp
--- a/Engine/Plugins/Experimental/CommonUI/Source/CommonUI/Private/Input/CommonBoundActionButton.cpp
+++ b/Engine/Plugins/Experimental/CommonUI/Source/CommonUI/Private/Input/CommonBoundActionButton.cpp
@@ -5,6 +5,7 @@
 #include "Input/UIActionBinding.h"
 #include "CommonActionWidget.h"
 #include "Framework/Application/SlateApplication.h"
+#include "Internationalization/LocKeyFuncs.h"
 
 #include UE_INLINE_GENERATED_CPP_BY_NAME(CommonBoundActionButton)
 
diff --git a/Engine/Plugins/Experimental/CommonUI/Source/CommonUI/Private/Input/CommonGenericInputActionDataTable.cpp b/Engine/Plugins/Experimental/CommonUI/Source/CommonUI/Private/Input/CommonGenericInputActionDataTable.cpp
--- a/Engine/Plugins/Experimental/CommonUI/Source/CommonUI/Private/Input/CommonGenericInputActionDataTable.cpp
+++ b/Engine/Plugins/Experimental/CommonUI/Source/CommonUI/Private/Input/CommonGenericInputActionDataTable.cpp
@@ -2,7 +2,6 @@
 
 #include "Input/CommonGenericInputActionDataTable.h"
 #include "CommonUITypes.h"
-#include "UObject/UObjectHash.h"
 
 #include UE_INLINE_GENERATED_CPP_BY_NAME(CommonGenericInputActionDataTable)
 
diff --git a/Engine/Plugins/Experimental/CommonUI/Source/CommonUI/Private/Input/CommonUIActionRouterBase.cpp b/Engine/Plugins/Experimental/CommonUI/Source/CommonUI/Private/Input/CommonUIActionRouterBase.cpp
--- a/Engine/Plugins/Experimental/CommonUI/Source/CommonUI/Private/Input/CommonUIActionRouterBase.cpp
+++ b/Engine/Plugins/Experimental/CommonUI/Source/CommonUI/Private/Input/CommonUIActionRouterBase.cpp
@@ -1,15 +1,14 @@
 // Copyright Epic Games, Inc. All Rights Reserved.
 
 #include "Input/CommonUIActionRouterBase.h"
+#include "CommonInputSubsystem.h"
 #include "Engine/GameInstance.h"
+#include "CommonInputTypeEnum.h"
 #include "Framework/Application/SlateApplication.h"
 #include "Input/CommonAnalogCursor.h"
 #include "Input/CommonUIInputSettings.h"
 #include "Input/UIActionBinding.h"
 #include "Input/UIActionRouterTypes.h"
-#include "CommonInputSettings.h"
-#include "CommonInputActionDomain.h"
-#include "ICommonInputModule.h"
 
 #include "Framework/Application/SlateUser.h"
 #include "Slate/SObjectWidget.h"
@@ -18,13 +17,10 @@
 #include "GameFramework/HUD.h"
 #include "Engine/Console.h"
 #include "Engine/Canvas.h"
-#include "Stats/Stats.h"
 
 #include "CommonGameViewportClient.h"
-#include "CommonUIPrivate.h"
 #include "CommonActivatableWidget.h"
 #include "CommonUIUtils.h"
-#include "CommonUISubsystemBase.h"
 #include "Slate/SGameLayerManager.h"
 #include "Framework/Commands/InputBindingManager.h"
 
diff --git a/Engine/Plugins/Experimental/CommonUI/Source/CommonUI/Private/Input/CommonUIInputSettings.cpp b/Engine/Plugins/Experimental/CommonUI/Source/CommonUI/Private/Input/CommonUIInputSettings.cpp
--- a/Engine/Plugins/Experimental/CommonUI/Source/CommonUI/Private/Input/CommonUIInputSettings.cpp
+++ b/Engine/Plugins/Experimental/CommonUI/Source/CommonUI/Private/Input/CommonUIInputSettings.cpp
@@ -1,7 +1,7 @@
 // Copyright Epic Games, Inc. All Rights Reserved.
 
 #include "Input/CommonUIInputSettings.h"
-#include "GameFramework/PlayerInput.h"
+#include "GameplayTagContainer.h"
 
 #include UE_INLINE_GENERATED_CPP_BY_NAME(CommonUIInputSettings)
 
diff --git a/Engine/Plugins/Experimental/CommonUI/Source/CommonUI/Private/Input/UIActionRouterTypes.cpp b/Engine/Plugins/Experimental/CommonUI/Source/CommonUI/Private/Input/UIActionRouterTypes.cpp
--- a/Engine/Plugins/Experimental/CommonUI/Source/CommonUI/Private/Input/UIActionRouterTypes.cpp
+++ b/Engine/Plugins/Experimental/CommonUI/Source/CommonUI/Private/Input/UIActionRouterTypes.cpp
@@ -1,10 +1,11 @@
 // Copyright Epic Games, Inc. All Rights Reserved.
 
 #include "Input/UIActionRouterTypes.h"
+#include "CommonInputTypeEnum.h"
+#include "CommonUITypes.h"
 #include "Input/CommonUIActionRouterBase.h"
 #include "CommonActivatableWidget.h"
 #include "Input/CommonUIInputTypes.h"
-#include "CommonUIPrivate.h"
 #include "CommonInputSettings.h"
 #include "CommonInputSubsystem.h"
 #include "ICommonInputModule.h"
@@ -12,6 +13,7 @@
 #include "Engine/GameViewportClient.h"
 #include "Framework/Application/SlateApplication.h"
 #include "Framework/Application/SlateUser.h"
+#include "Input/UIActionBinding.h"
 #include "Widgets/SViewport.h"
 #include "CommonButtonBase.h"
 
diff --git a/Engine/Plugins/Experimental/CommonUI/Source/CommonUI/Private/SAnalogSlider.cpp b/Engine/Plugins/Experimental/CommonUI/Source/CommonUI/Private/SAnalogSlider.cpp
--- a/Engine/Plugins/Experimental/CommonUI/Source/CommonUI/Private/SAnalogSlider.cpp
+++ b/Engine/Plugins/Experimental/CommonUI/Source/CommonUI/Private/SAnalogSlider.cpp
@@ -1,7 +1,6 @@
 // Copyright Epic Games, Inc. All Rights Reserved.
 
 #include "SAnalogSlider.h"
-#include "Rendering/DrawElements.h"
 #include "Misc/App.h"
 
 void SAnalogSlider::Construct(const SAnalogSlider::FArguments& InDeclaration)
diff --git a/Engine/Plugins/Experimental/CommonUI/Source/CommonUI/Private/SCommonHierarchicalScrollBox.cpp b/Engine/Plugins/Experimental/CommonUI/Source/CommonUI/Private/SCommonHierarchicalScrollBox.cpp
--- a/Engine/Plugins/Experimental/CommonUI/Source/CommonUI/Private/SCommonHierarchicalScrollBox.cpp
+++ b/Engine/Plugins/Experimental/CommonUI/Source/CommonUI/Private/SCommonHierarchicalScrollBox.cpp
@@ -1,10 +1,7 @@
 // Copyright Epic Games, Inc. All Rights Reserved.
 
 #include "SCommonHierarchicalScrollBox.h"
-#include "Types/SlateConstants.h"
 #include "Types/NavigationMetaData.h"
-#include "Input/HittestGrid.h"
-#include "Framework/Application/SlateApplication.h"
 
 void SCommonHierarchicalScrollBox::AppendFocusableChildren(TArray<TSharedRef<SWidget>>& OutChildren, TSharedRef<SWidget> Owner)
 {
@@ -155,4 +152,4 @@ FNavigationReply SCommonHierarchicalScrollBox::OnNavigation(const FGeometry& MyG
 	}
 
 	return SCompoundWidget::OnNavigation(MyGeometry, InNavigationEvent);
-}
\ No newline at end of file
+}
diff --git a/Engine/Plugins/Experimental/CommonUI/Source/CommonUI/Private/StyleSheet/CommonStyleSheet.cpp b/Engine/Plugins/Experimental/CommonUI/Source/CommonUI/Private/StyleSheet/CommonStyleSheet.cpp
--- a/Engine/Plugins/Experimental/CommonUI/Source/CommonUI/Private/StyleSheet/CommonStyleSheet.cpp
+++ b/Engine/Plugins/Experimental/CommonUI/Source/CommonUI/Private/StyleSheet/CommonStyleSheet.cpp
@@ -3,9 +3,7 @@
 #include "StyleSheet/CommonStyleSheet.h"
 #include "StyleSheet/CommonStyleSheetTypes.h"
 
-#include "Blueprint/UserWidget.h"
 #include "CommonTextBlock.h"
-#include "Components/Widget.h"
 
 #include UE_INLINE_GENERATED_CPP_BY_NAME(CommonStyleSheet)
 
diff --git a/Engine/Plugins/Experimental/CommonUI/Source/CommonUI/Private/Widgets/CommonActivatableWidgetContainer.cpp b/Engine/Plugins/Experimental/CommonUI/Source/CommonUI/Private/Widgets/CommonActivatableWidgetContainer.cpp
--- a/Engine/Plugins/Experimental/CommonUI/Source/CommonUI/Private/Widgets/CommonActivatableWidgetContainer.cpp
+++ b/Engine/Plugins/Experimental/CommonUI/Source/CommonUI/Private/Widgets/CommonActivatableWidgetContainer.cpp
@@ -1,13 +1,10 @@
 // Copyright Epic Games, Inc. All Rights Reserved.
 
 #include "Widgets/CommonActivatableWidgetContainer.h"
-#include "CommonWidgetPaletteCategories.h"
 #include "CommonActivatableWidget.h"
-#include "CommonUIUtils.h"
-#include "Input/CommonUIActionRouterBase.h"
+#include "Slate/SCommonAnimatedSwitcher.h"
 #include "Widgets/SOverlay.h"
 #include "Widgets/Layout/SSpacer.h"
-#include "Stats/Stats.h"
 #include "TimerManager.h"
 
 #include UE_INLINE_GENERATED_CPP_BY_NAME(CommonActivatableWidgetContainer)
diff --git a/Engine/Plugins/Experimental/CommonUI/Source/CommonUIEditor/Private/CommonAssetTypeActions_GenericInputActionDataTable.cpp b/Engine/Plugins/Experimental/CommonUI/Source/CommonUIEditor/Private/CommonAssetTypeActions_GenericInputActionDataTable.cpp
--- a/Engine/Plugins/Experimental/CommonUI/Source/CommonUIEditor/Private/CommonAssetTypeActions_GenericInputActionDataTable.cpp
+++ b/Engine/Plugins/Experimental/CommonUI/Source/CommonUIEditor/Private/CommonAssetTypeActions_GenericInputActionDataTable.cpp
@@ -4,7 +4,6 @@
 #include "DataTableEditorModule.h"
 #include "Input/CommonGenericInputActionDataTable.h"
 #include "Misc/MessageDialog.h"
-#include "ToolMenus.h"
 
 #define LOCTEXT_NAMESPACE "AssetTypeActions"
 
diff --git a/Engine/Plugins/Experimental/CommonUI/Source/CommonUIEditor/Private/CommonUIEditorModule.cpp b/Engine/Plugins/Experimental/CommonUI/Source/CommonUIEditor/Private/CommonUIEditorModule.cpp
--- a/Engine/Plugins/Experimental/CommonUI/Source/CommonUIEditor/Private/CommonUIEditorModule.cpp
+++ b/Engine/Plugins/Experimental/CommonUI/Source/CommonUIEditor/Private/CommonUIEditorModule.cpp
@@ -1,13 +1,10 @@
 // Copyright Epic Games, Inc. All Rights Reserved.
 
 #include "CommonUIEditorModule.h"
-#include "CommonUIPrivate.h"
 #include "CommonVideoPlayerCustomization.h"
 
 #include "GameplayTagsEditorModule.h"
 #include "CommonAssetTypeActions_GenericInputActionDataTable.h"
-#include "Modules/ModuleInterface.h"
-#include "Modules/ModuleManager.h"
 #include "PropertyEditorModule.h"
 #include "UITag.h"
 
diff --git a/Engine/Plugins/Experimental/CommonUI/Source/CommonUIEditor/Private/CommonVideoPlayerCustomization.cpp b/Engine/Plugins/Experimental/CommonUI/Source/CommonUIEditor/Private/CommonVideoPlayerCustomization.cpp
--- a/Engine/Plugins/Experimental/CommonUI/Source/CommonUIEditor/Private/CommonVideoPlayerCustomization.cpp
+++ b/Engine/Plugins/Experimental/CommonUI/Source/CommonUIEditor/Private/CommonVideoPlayerCustomization.cpp
@@ -1,15 +1,11 @@
 // Copyright Epic Games, Inc. All Rights Reserved.
 
 #include "CommonVideoPlayerCustomization.h"
-#include "PropertyCustomizationHelpers.h"
 #include "CommonVideoPlayer.h"
-#include "ITransportControl.h"
+#include "DetailWidgetRow.h"
 #include "EditorWidgetsModule.h"
 #include "Widgets/Input/SButton.h"
-#include "Widgets/Images/SImage.h"
-#include "Styling/AppStyle.h"
 #include "DetailLayoutBuilder.h"
-#include "DetailCategoryBuilder.h"
 #include "Widgets/Input/SNumericEntryBox.h"
 #include "Modules/ModuleManager.h"
 
