[CommonUI]
* Ran IWYU on h and cpp and put removed headers inside UE_ENABLE_INCLUDE_ORDER_DEPRECATED_IN_5_2 scope

diff --git a/Engine/Plugins/Experimental/CommonUI/Source/CommonUI/Private/CommonButtonBase.cpp b/Engine/Plugins/Experimental/CommonUI/Source/CommonUI/Private/CommonButtonBase.cpp
--- a/Engine/Plugins/Experimental/CommonUI/Source/CommonUI/Private/CommonButtonBase.cpp
+++ b/Engine/Plugins/Experimental/CommonUI/Source/CommonUI/Private/CommonButtonBase.cpp
@@ -3,7 +3,6 @@
 #include "CommonButtonBase.h"
 
 #include "CommonActionWidget.h"
-#include "CommonInputTypeEnum.h"
 #include "CommonUISubsystemBase.h"
 #include "CommonInputSubsystem.h"
 #include "CommonUIEditorSettings.h"
diff --git a/Engine/Plugins/Experimental/CommonUI/Source/CommonUI/Private/CommonButtonTypes.cpp b/Engine/Plugins/Experimental/CommonUI/Source/CommonUI/Private/CommonButtonTypes.cpp
--- a/Engine/Plugins/Experimental/CommonUI/Source/CommonUI/Private/CommonButtonTypes.cpp
+++ b/Engine/Plugins/Experimental/CommonUI/Source/CommonUI/Private/CommonButtonTypes.cpp
@@ -1,6 +1,7 @@
 // Copyright Epic Games, Inc. All Rights Reserved.
 
 #include "CommonButtonTypes.h"
+#include "Input/Reply.h"
 
 //////////////////////////////////////////////////////////////////////////
 // SCommonButton
@@ -205,4 +206,4 @@ int32 SCommonButton::OnPaint(const FPaintArgs& Args, const FGeometry& AllottedGe
 {
 	bool bEnabled = bParentEnabled && bIsButtonEnabled;
 	return SButton::OnPaint(Args, AllottedGeometry, MyClippingRect, OutDrawElements, LayerId, InWidgetStyle, bEnabled);
-}
\ No newline at end of file
+}
diff --git a/Engine/Plugins/Experimental/CommonUI/Source/CommonUI/Private/CommonCustomVersion.h b/Engine/Plugins/Experimental/CommonUI/Source/CommonUI/Private/CommonCustomVersion.h
--- a/Engine/Plugins/Experimental/CommonUI/Source/CommonUI/Private/CommonCustomVersion.h
+++ b/Engine/Plugins/Experimental/CommonUI/Source/CommonUI/Private/CommonCustomVersion.h
@@ -2,8 +2,8 @@
 
 #pragma once
 
-#include "CoreMinimal.h"
-#include "Misc/Guid.h"
+
+struct FGuid;
 
 // Custom serialization version for all packages containing CommonUI asset types
 struct FCommonCustomVersion
@@ -27,3 +27,8 @@ struct FCommonCustomVersion
 private:
 	FCommonCustomVersion() {}
 };
+
+#if UE_ENABLE_INCLUDE_ORDER_DEPRECATED_IN_5_2
+#include "CoreMinimal.h"
+#include "Misc/Guid.h"
+#endif
diff --git a/Engine/Plugins/Experimental/CommonUI/Source/CommonUI/Private/CommonUIObjectVersion.h b/Engine/Plugins/Experimental/CommonUI/Source/CommonUI/Private/CommonUIObjectVersion.h
--- a/Engine/Plugins/Experimental/CommonUI/Source/CommonUI/Private/CommonUIObjectVersion.h
+++ b/Engine/Plugins/Experimental/CommonUI/Source/CommonUI/Private/CommonUIObjectVersion.h
@@ -1,7 +1,6 @@
 // Copyright Epic Games, Inc. All Rights Reserved.
 #pragma once
 
-#include "CoreTypes.h"
 #include "Misc/Guid.h"
 
 // Custom serialization version for changes made to the CommonUI module
@@ -22,3 +21,7 @@ namespace ECommonUIObjectVersion
 	// The GUID for this custom version number
 	const extern FGuid Guid;
 };
+
+#if UE_ENABLE_INCLUDE_ORDER_DEPRECATED_IN_5_2
+#include "CoreTypes.h"
+#endif
diff --git a/Engine/Plugins/Experimental/CommonUI/Source/CommonUI/Private/CommonUIPrivate.h b/Engine/Plugins/Experimental/CommonUI/Source/CommonUI/Private/CommonUIPrivate.h
--- a/Engine/Plugins/Experimental/CommonUI/Source/CommonUI/Private/CommonUIPrivate.h
+++ b/Engine/Plugins/Experimental/CommonUI/Source/CommonUI/Private/CommonUIPrivate.h
@@ -6,6 +6,10 @@
 /* Dependencies
  *****************************************************************************/
 
-#include "Styling/SlateStyle.h"
+#include "Logging/LogMacros.h"
 
 DECLARE_LOG_CATEGORY_EXTERN(LogCommonUI, Log, All);
+
+#if UE_ENABLE_INCLUDE_ORDER_DEPRECATED_IN_5_2
+#include "Styling/SlateStyle.h"
+#endif
diff --git a/Engine/Plugins/Experimental/CommonUI/Source/CommonUI/Private/CommonUISubsystemBase.cpp b/Engine/Plugins/Experimental/CommonUI/Source/CommonUI/Private/CommonUISubsystemBase.cpp
--- a/Engine/Plugins/Experimental/CommonUI/Source/CommonUI/Private/CommonUISubsystemBase.cpp
+++ b/Engine/Plugins/Experimental/CommonUI/Source/CommonUI/Private/CommonUISubsystemBase.cpp
@@ -5,7 +5,6 @@
 
 #include "CommonInputBaseTypes.h"
 #include "CommonInputSubsystem.h"
-#include "CommonInputTypeEnum.h"
 #include "CommonUITypes.h"
 #include "IAnalyticsProviderET.h"
 #include "Components/Widget.h"
diff --git a/Engine/Plugins/Experimental/CommonUI/Source/CommonUI/Private/CommonUITypes.cpp b/Engine/Plugins/Experimental/CommonUI/Source/CommonUI/Private/CommonUITypes.cpp
--- a/Engine/Plugins/Experimental/CommonUI/Source/CommonUI/Private/CommonUITypes.cpp
+++ b/Engine/Plugins/Experimental/CommonUI/Source/CommonUI/Private/CommonUITypes.cpp
@@ -2,7 +2,6 @@
  
 #include "CommonUITypes.h"
 #include "CommonInputBaseTypes.h"
-#include "CommonInputTypeEnum.h"
 #include "CommonUIPrivate.h"
 #include "CommonInputSubsystem.h"
 #include "HAL/PlatformInput.h"
diff --git a/Engine/Plugins/Experimental/CommonUI/Source/CommonUI/Private/CommonVisibilityWidgetBase.cpp b/Engine/Plugins/Experimental/CommonUI/Source/CommonUI/Private/CommonVisibilityWidgetBase.cpp
--- a/Engine/Plugins/Experimental/CommonUI/Source/CommonUI/Private/CommonVisibilityWidgetBase.cpp
+++ b/Engine/Plugins/Experimental/CommonUI/Source/CommonUI/Private/CommonVisibilityWidgetBase.cpp
@@ -2,7 +2,6 @@
 
 #include "CommonVisibilityWidgetBase.h"
 #include "CommonInputSubsystem.h"
-#include "CommonInputTypeEnum.h"
 #include "CommonUIPrivate.h"
 
 #include "CommonInputBaseTypes.h"
diff --git a/Engine/Plugins/Experimental/CommonUI/Source/CommonUI/Private/Input/UIActionRouterTypes.cpp b/Engine/Plugins/Experimental/CommonUI/Source/CommonUI/Private/Input/UIActionRouterTypes.cpp
--- a/Engine/Plugins/Experimental/CommonUI/Source/CommonUI/Private/Input/UIActionRouterTypes.cpp
+++ b/Engine/Plugins/Experimental/CommonUI/Source/CommonUI/Private/Input/UIActionRouterTypes.cpp
@@ -1,7 +1,6 @@
 // Copyright Epic Games, Inc. All Rights Reserved.
 
 #include "Input/UIActionRouterTypes.h"
-#include "CommonInputTypeEnum.h"
 #include "CommonUITypes.h"
 #include "Input/CommonUIActionRouterBase.h"
 #include "CommonActivatableWidget.h"
diff --git a/Engine/Plugins/Experimental/CommonUI/Source/CommonUI/Private/Input/UIActionRouterTypes.h b/Engine/Plugins/Experimental/CommonUI/Source/CommonUI/Private/Input/UIActionRouterTypes.h
--- a/Engine/Plugins/Experimental/CommonUI/Source/CommonUI/Private/Input/UIActionRouterTypes.h
+++ b/Engine/Plugins/Experimental/CommonUI/Source/CommonUI/Private/Input/UIActionRouterTypes.h
@@ -3,8 +3,14 @@
 #pragma once
 
 #include "Input/UIActionBindingHandle.h"
-#include "Input/UIActionBinding.h"
-#include "Input/CommonUIInputSettings.h"
+#include "UObject/WeakObjectPtr.h"
+#include "UObject/WeakObjectPtrTemplates.h"
+
+enum EInputEvent : int;
+enum class EProcessHoldActionResult;
+struct FKey;
+struct FUIActionBinding;
+template <typename OptionalType> struct TOptional;
 
 // Note: Everything in here should be considered completely private to each other and CommonUIActionRouter.
 //		They were all originally defined directly in CommonUIActionRouter.cpp, but it was annoying having to scroll around so much.
@@ -215,4 +221,9 @@ private:
 
 	// WeakPtr because the root itself can be the primary active node - results in a circular ref leak using a full SharedPtr here
 	TWeakPtr<FActivatableTreeNode> LeafmostActiveNode;
-};
\ No newline at end of file
+};
+
+#if UE_ENABLE_INCLUDE_ORDER_DEPRECATED_IN_5_2
+#include "Input/CommonUIInputSettings.h"
+#include "Input/UIActionBinding.h"
+#endif
diff --git a/Engine/Plugins/Experimental/CommonUI/Source/CommonUI/Private/SAnalogSlider.cpp b/Engine/Plugins/Experimental/CommonUI/Source/CommonUI/Private/SAnalogSlider.cpp
--- a/Engine/Plugins/Experimental/CommonUI/Source/CommonUI/Private/SAnalogSlider.cpp
+++ b/Engine/Plugins/Experimental/CommonUI/Source/CommonUI/Private/SAnalogSlider.cpp
@@ -1,7 +1,9 @@
 // Copyright Epic Games, Inc. All Rights Reserved.
 
 #include "SAnalogSlider.h"
+#include "Input/NavigationReply.h"
 #include "Misc/App.h"
+#include "Input/Reply.h"
 
 void SAnalogSlider::Construct(const SAnalogSlider::FArguments& InDeclaration)
 {
diff --git a/Engine/Plugins/Experimental/CommonUI/Source/CommonUI/Public/AnalogSlider.h b/Engine/Plugins/Experimental/CommonUI/Source/CommonUI/Public/AnalogSlider.h
--- a/Engine/Plugins/Experimental/CommonUI/Source/CommonUI/Public/AnalogSlider.h
+++ b/Engine/Plugins/Experimental/CommonUI/Source/CommonUI/Public/AnalogSlider.h
@@ -3,13 +3,12 @@
 #pragma once
 
 #include "Components/Slider.h"
-#include "SAnalogSlider.h"
-#include "CommonInputSubsystem.h"
 
-#include "Types/NavigationMetaData.h"
 
 #include "AnalogSlider.generated.h"
 
+enum class ECommonInputType : uint8;
+
 class SAnalogSlider;
 
 /**
@@ -41,3 +40,9 @@ public:
 protected:
 	TSharedPtr<SAnalogSlider> MyAnalogSlider;
 };
+
+#if UE_ENABLE_INCLUDE_ORDER_DEPRECATED_IN_5_2
+#include "CommonInputSubsystem.h"
+#include "SAnalogSlider.h"
+#include "Types/NavigationMetaData.h"
+#endif
diff --git a/Engine/Plugins/Experimental/CommonUI/Source/CommonUI/Public/CommonActionHandlerInterface.h b/Engine/Plugins/Experimental/CommonUI/Source/CommonUI/Public/CommonActionHandlerInterface.h
--- a/Engine/Plugins/Experimental/CommonUI/Source/CommonUI/Public/CommonActionHandlerInterface.h
+++ b/Engine/Plugins/Experimental/CommonUI/Source/CommonUI/Public/CommonActionHandlerInterface.h
@@ -4,8 +4,14 @@
 
 #include "CommonUITypes.h"
 
+#include "Components/Widget.h"
+#include "UObject/Interface.h"
 #include "CommonActionHandlerInterface.generated.h"
 
+enum EInputEvent : int;
+namespace ETouchType { enum Type : int; }
+struct FKey;
+
 /** Action Delegates */
 
 /** 
@@ -132,4 +138,4 @@ public:
 	 *  @param bHoldAction if true, we're looking for a hold action specifically.
 	 */
 	virtual void TriggerFirstMatchingInputAction(int32 ControllerId, const FCommonInputActionDataBase& InInputActionData, bool bHoldAction) = 0;
-};
\ No newline at end of file
+};
diff --git a/Engine/Plugins/Experimental/CommonUI/Source/CommonUI/Public/CommonActionWidget.h b/Engine/Plugins/Experimental/CommonUI/Source/CommonUI/Public/CommonActionWidget.h
--- a/Engine/Plugins/Experimental/CommonUI/Source/CommonUI/Public/CommonActionWidget.h
+++ b/Engine/Plugins/Experimental/CommonUI/Source/CommonUI/Public/CommonActionWidget.h
@@ -2,11 +2,16 @@
 
 #pragma once
 
-#include "CommonUITypes.h"
+#include "Components/Widget.h"
 #include "Input/UIActionBindingHandle.h"
 
+#include "Engine/DataTable.h"
 #include "CommonActionWidget.generated.h"
 
+class UCommonInputSubsystem;
+enum class ECommonInputType : uint8;
+struct FCommonInputActionDataBase;
+
 class SBox;
 class SImage;
 class UMaterialInstanceDynamic;
@@ -122,4 +127,8 @@ protected:
 	FSlateBrush Icon;
 
 	bool bAlwaysHideOverride = false;
-};
\ No newline at end of file
+};
+
+#if UE_ENABLE_INCLUDE_ORDER_DEPRECATED_IN_5_2
+#include "CommonUITypes.h"
+#endif
diff --git a/Engine/Plugins/Experimental/CommonUI/Source/CommonUI/Public/CommonActivatableWidget.h b/Engine/Plugins/Experimental/CommonUI/Source/CommonUI/Public/CommonActivatableWidget.h
--- a/Engine/Plugins/Experimental/CommonUI/Source/CommonUI/Public/CommonActivatableWidget.h
+++ b/Engine/Plugins/Experimental/CommonUI/Source/CommonUI/Public/CommonActivatableWidget.h
@@ -3,7 +3,6 @@
 #pragma once
 
 #include "CommonUserWidget.h"
-#include "Input/UIActionBindingHandle.h"
 #include "CommonInputActionDomain.h"
 #include "CommonActivatableWidget.generated.h"
 
@@ -252,4 +251,8 @@ protected:
 	virtual void InternalProcessActivation();
 	virtual void InternalProcessDeactivation();
 	void Reset();
-};
\ No newline at end of file
+};
+
+#if UE_ENABLE_INCLUDE_ORDER_DEPRECATED_IN_5_2
+#include "Input/UIActionBindingHandle.h"
+#endif
diff --git a/Engine/Plugins/Experimental/CommonUI/Source/CommonUI/Public/CommonActivatableWidgetSwitcher.h b/Engine/Plugins/Experimental/CommonUI/Source/CommonUI/Public/CommonActivatableWidgetSwitcher.h
--- a/Engine/Plugins/Experimental/CommonUI/Source/CommonUI/Public/CommonActivatableWidgetSwitcher.h
+++ b/Engine/Plugins/Experimental/CommonUI/Source/CommonUI/Public/CommonActivatableWidgetSwitcher.h
@@ -4,8 +4,6 @@
 
 #include "CommonAnimatedSwitcher.h"
 
-#include "Misc/Optional.h"
-#include "UObject/WeakObjectPtrTemplates.h"
 
 #include "CommonActivatableWidgetSwitcher.generated.h"
 
@@ -41,3 +39,8 @@ private:
 
 	TOptional<TWeakObjectPtr<UCommonActivatableWidget>> WeakOwningActivatableWidget;
 };
+
+#if UE_ENABLE_INCLUDE_ORDER_DEPRECATED_IN_5_2
+#include "Misc/Optional.h"
+#include "UObject/WeakObjectPtrTemplates.h"
+#endif
diff --git a/Engine/Plugins/Experimental/CommonUI/Source/CommonUI/Public/CommonAnimatedSwitcher.h b/Engine/Plugins/Experimental/CommonUI/Source/CommonUI/Public/CommonAnimatedSwitcher.h
--- a/Engine/Plugins/Experimental/CommonUI/Source/CommonUI/Public/CommonAnimatedSwitcher.h
+++ b/Engine/Plugins/Experimental/CommonUI/Source/CommonUI/Public/CommonAnimatedSwitcher.h
@@ -4,7 +4,6 @@
 
 #include "Components/WidgetSwitcher.h"
 #include "Slate/SCommonAnimatedSwitcher.h"
-#include "Components/Widget.h"
 
 #include "CommonAnimatedSwitcher.generated.h"
 
@@ -93,4 +92,8 @@ protected:
 
 private:
 	void SetActiveWidgetIndex_Internal(int32 Index);
-};
\ No newline at end of file
+};
+
+#if UE_ENABLE_INCLUDE_ORDER_DEPRECATED_IN_5_2
+#include "Components/Widget.h"
+#endif
diff --git a/Engine/Plugins/Experimental/CommonUI/Source/CommonUI/Public/CommonBorder.h b/Engine/Plugins/Experimental/CommonUI/Source/CommonUI/Public/CommonBorder.h
--- a/Engine/Plugins/Experimental/CommonUI/Source/CommonUI/Public/CommonBorder.h
+++ b/Engine/Plugins/Experimental/CommonUI/Source/CommonUI/Public/CommonBorder.h
@@ -2,10 +2,11 @@
 
 #pragma once
 
-#include "CommonUITypes.h"
 #include "Components/Border.h"
 #include "CommonBorder.generated.h"
 
+struct FDesignerChangedEventArgs;
+
 /* 
  * ---- All properties must be EditDefaultsOnly, BlueprintReadOnly !!! -----
  * We return the CDO to blueprints, so we cannot allow any changes (blueprint doesn't support const variables)
@@ -82,4 +83,8 @@ protected:
 private:
 	const UCommonBorderStyle* GetStyleCDO() const;
 
-};
\ No newline at end of file
+};
+
+#if UE_ENABLE_INCLUDE_ORDER_DEPRECATED_IN_5_2
+#include "CommonUITypes.h"
+#endif
diff --git a/Engine/Plugins/Experimental/CommonUI/Source/CommonUI/Public/CommonButtonBase.h b/Engine/Plugins/Experimental/CommonUI/Source/CommonUI/Public/CommonButtonBase.h
--- a/Engine/Plugins/Experimental/CommonUI/Source/CommonUI/Public/CommonButtonBase.h
+++ b/Engine/Plugins/Experimental/CommonUI/Source/CommonUI/Public/CommonButtonBase.h
@@ -3,13 +3,16 @@
 #pragma once
 
 #include "CommonUserWidget.h"
-#include "CommonUITypes.h"
 #include "Components/Button.h"
-#include "Misc/Optional.h"
+#include "Engine/DataTable.h"
 #include "Types/ISlateMetaData.h"
 #include "CommonInputModeTypes.h"
 #include "CommonButtonBase.generated.h"
 
+class UCommonButtonBase;
+class UMaterialInstanceDynamic;
+enum class ECommonInputType : uint8;
+
 class UCommonTextStyle;
 class SBox;
 class UCommonActionWidget;
@@ -859,4 +862,9 @@ protected:
 	 */
 	UPROPERTY(BlueprintReadOnly, Category = Input, meta = (BindWidget, OptionalWidget = true, AllowPrivateAccess = true))
 	TObjectPtr<UCommonActionWidget> InputActionWidget;
-};
\ No newline at end of file
+};
+
+#if UE_ENABLE_INCLUDE_ORDER_DEPRECATED_IN_5_2
+#include "CommonUITypes.h"
+#include "Misc/Optional.h"
+#endif
diff --git a/Engine/Plugins/Experimental/CommonUI/Source/CommonUI/Public/CommonCustomNavigation.h b/Engine/Plugins/Experimental/CommonUI/Source/CommonUI/Public/CommonCustomNavigation.h
--- a/Engine/Plugins/Experimental/CommonUI/Source/CommonUI/Public/CommonCustomNavigation.h
+++ b/Engine/Plugins/Experimental/CommonUI/Source/CommonUI/Public/CommonCustomNavigation.h
@@ -2,7 +2,6 @@
 
 #pragma once
 
-#include "CommonUITypes.h"
 #include "Components/Border.h"
 #include "CommonCustomNavigation.generated.h"
 
@@ -37,4 +36,8 @@ protected:
 	
 	bool OnNavigation(EUINavigation NavigationType);
 
-};
\ No newline at end of file
+};
+
+#if UE_ENABLE_INCLUDE_ORDER_DEPRECATED_IN_5_2
+#include "CommonUITypes.h"
+#endif
diff --git a/Engine/Plugins/Experimental/CommonUI/Source/CommonUI/Public/CommonGameViewportClient.h b/Engine/Plugins/Experimental/CommonUI/Source/CommonUI/Public/CommonGameViewportClient.h
--- a/Engine/Plugins/Experimental/CommonUI/Source/CommonUI/Public/CommonGameViewportClient.h
+++ b/Engine/Plugins/Experimental/CommonUI/Source/CommonUI/Public/CommonGameViewportClient.h
@@ -2,14 +2,11 @@
 
 #pragma once
 
-#include "CoreMinimal.h"
-#include "UObject/ObjectMacros.h"
-#include "InputCoreTypes.h"
-#include "InputKeyEventArgs.h"
-#include "Input/Reply.h"
 #include "Engine/GameViewportClient.h"
 #include "CommonGameViewportClient.generated.h"
 
+class FReply;
+
 DECLARE_DELEGATE_FourParams(FOnRerouteInputDelegate, FInputDeviceId /* InputDeviceId */, FKey /* Key */, EInputEvent /* EventType */, FReply& /* Reply */);
 DECLARE_DELEGATE_FourParams(FOnRerouteAxisDelegate, FInputDeviceId /* InputDeviceId */, FKey /* Key */, float /* Delta */, FReply& /* Reply */);
 DECLARE_DELEGATE_FiveParams(FOnRerouteTouchDelegate, int32 /* ControllerId */, uint32 /* TouchId */, ETouchType::Type /* TouchType */, const FVector2D& /* TouchLocation */, FReply& /* Reply */);
@@ -67,3 +64,11 @@ protected:
 
 	FOnRerouteInputDelegate RerouteBlockedInput;
 };
+
+#if UE_ENABLE_INCLUDE_ORDER_DEPRECATED_IN_5_2
+#include "CoreMinimal.h"
+#include "Input/Reply.h"
+#include "InputCoreTypes.h"
+#include "InputKeyEventArgs.h"
+#include "UObject/ObjectMacros.h"
+#endif
diff --git a/Engine/Plugins/Experimental/CommonUI/Source/CommonUI/Public/CommonHardwareVisibilityBorder.h b/Engine/Plugins/Experimental/CommonUI/Source/CommonUI/Public/CommonHardwareVisibilityBorder.h
--- a/Engine/Plugins/Experimental/CommonUI/Source/CommonUI/Public/CommonHardwareVisibilityBorder.h
+++ b/Engine/Plugins/Experimental/CommonUI/Source/CommonUI/Public/CommonHardwareVisibilityBorder.h
@@ -2,12 +2,12 @@
 
 #pragma once
 
-#include "CoreMinimal.h"
-#include "CommonUITypes.h"
 #include "CommonBorder.h"
 #include "GameplayTagContainer.h"
 #include "CommonHardwareVisibilityBorder.generated.h"
 
+enum class ESlateVisibility : uint8;
+
 class UCommonUIVisibilitySubsystem;
 
 /**
@@ -43,3 +43,8 @@ protected:
 
 	void HandleInputMethodChanged(UCommonUIVisibilitySubsystem*);
 };
+
+#if UE_ENABLE_INCLUDE_ORDER_DEPRECATED_IN_5_2
+#include "CommonUITypes.h"
+#include "CoreMinimal.h"
+#endif
diff --git a/Engine/Plugins/Experimental/CommonUI/Source/CommonUI/Public/CommonHierarchicalScrollBox.h b/Engine/Plugins/Experimental/CommonUI/Source/CommonUI/Public/CommonHierarchicalScrollBox.h
--- a/Engine/Plugins/Experimental/CommonUI/Source/CommonUI/Public/CommonHierarchicalScrollBox.h
+++ b/Engine/Plugins/Experimental/CommonUI/Source/CommonUI/Public/CommonHierarchicalScrollBox.h
@@ -2,8 +2,6 @@
 
 #pragma once
 
-#include "CoreMinimal.h"
-#include "UObject/ObjectMacros.h"
 #include "Components/ScrollBox.h"
 #include "CommonHierarchicalScrollBox.generated.h"
 
@@ -20,3 +18,8 @@ protected:
 	virtual TSharedRef<SWidget> RebuildWidget() override;
 	//~ End UWidget Interface
 };
+
+#if UE_ENABLE_INCLUDE_ORDER_DEPRECATED_IN_5_2
+#include "CoreMinimal.h"
+#include "UObject/ObjectMacros.h"
+#endif
diff --git a/Engine/Plugins/Experimental/CommonUI/Source/CommonUI/Public/CommonLazyImage.h b/Engine/Plugins/Experimental/CommonUI/Source/CommonUI/Public/CommonLazyImage.h
--- a/Engine/Plugins/Experimental/CommonUI/Source/CommonUI/Public/CommonLazyImage.h
+++ b/Engine/Plugins/Experimental/CommonUI/Source/CommonUI/Public/CommonLazyImage.h
@@ -7,6 +7,8 @@
 
 #include "CommonLazyImage.generated.h"
 
+class UTexture;
+
 class UCommonMcpItemDefinition;
 
 /**
@@ -103,4 +105,4 @@ private:
 
 	TSharedPtr<SLoadGuard> MyLoadGuard;
 	FOnLoadGuardStateChangedEvent OnLoadingStateChangedEvent;
-};
\ No newline at end of file
+};
diff --git a/Engine/Plugins/Experimental/CommonUI/Source/CommonUI/Public/CommonLazyWidget.h b/Engine/Plugins/Experimental/CommonUI/Source/CommonUI/Public/CommonLazyWidget.h
--- a/Engine/Plugins/Experimental/CommonUI/Source/CommonUI/Public/CommonLazyWidget.h
+++ b/Engine/Plugins/Experimental/CommonUI/Source/CommonUI/Public/CommonLazyWidget.h
@@ -2,9 +2,10 @@
 
 #pragma once
 
+#include "Blueprint/UserWidget.h"
 #include "CommonLoadGuard.h"
-#include "Components/Widget.h"
 
+#include "Engine/StreamableManager.h"
 #include "CommonLazyWidget.generated.h"
 
 class UCommonMcpItemDefinition;
@@ -72,4 +73,8 @@ private:
 	FOnLoadGuardStateChangedEvent OnLoadingStateChangedEvent;
 
 	FOnLazyContentChangedEvent OnContentChangedEvent;
-};
\ No newline at end of file
+};
+
+#if UE_ENABLE_INCLUDE_ORDER_DEPRECATED_IN_5_2
+#include "Components/Widget.h"
+#endif
diff --git a/Engine/Plugins/Experimental/CommonUI/Source/CommonUI/Public/CommonListView.h b/Engine/Plugins/Experimental/CommonUI/Source/CommonUI/Public/CommonListView.h
--- a/Engine/Plugins/Experimental/CommonUI/Source/CommonUI/Public/CommonListView.h
+++ b/Engine/Plugins/Experimental/CommonUI/Source/CommonUI/Public/CommonListView.h
@@ -5,6 +5,8 @@
 #include "Components/ListView.h"
 #include "CommonListView.generated.h"
 
+class STableViewBase;
+
 //////////////////////////////////////////////////////////////////////////
 // SCommonListView
 //////////////////////////////////////////////////////////////////////////
@@ -106,4 +108,4 @@ public:
 protected:
 	virtual TSharedRef<STableViewBase> RebuildListWidget() override;
 	virtual UUserWidget& OnGenerateEntryWidgetInternal(UObject* Item, TSubclassOf<UUserWidget> DesiredEntryClass, const TSharedRef<STableViewBase>& OwnerTable) override;
-};
\ No newline at end of file
+};
diff --git a/Engine/Plugins/Experimental/CommonUI/Source/CommonUI/Public/CommonLoadGuard.h b/Engine/Plugins/Experimental/CommonUI/Source/CommonUI/Public/CommonLoadGuard.h
--- a/Engine/Plugins/Experimental/CommonUI/Source/CommonUI/Public/CommonLoadGuard.h
+++ b/Engine/Plugins/Experimental/CommonUI/Source/CommonUI/Public/CommonLoadGuard.h
@@ -2,16 +2,15 @@
 
 #pragma once
 
-#include "CommonUITypes.h"
-#include "ICommonUIModule.h"
-#include "Engine/StreamableManager.h"
 #include "Components/ContentWidget.h"
+#include "Widgets/Accessibility/SlateWidgetAccessibleTypes.h"
 #include "Widgets/SCompoundWidget.h"
-#include "Animation/CurveSequence.h"
-#include "Engine/AssetManager.h"
 
 #include "CommonLoadGuard.generated.h"
 
+struct FStreamableHandle;
+struct FTextBlockStyle;
+
 class UCommonTextStyle;
 class STextBlock;
 class SBorder;
@@ -248,3 +247,11 @@ private:
 
 	FOnLoadGuardStateChangedEvent OnLoadingStateChangedEvent;
 };
+
+#if UE_ENABLE_INCLUDE_ORDER_DEPRECATED_IN_5_2
+#include "Animation/CurveSequence.h"
+#include "CommonUITypes.h"
+#include "Engine/AssetManager.h"
+#include "Engine/StreamableManager.h"
+#include "ICommonUIModule.h"
+#endif
diff --git a/Engine/Plugins/Experimental/CommonUI/Source/CommonUI/Public/CommonNativeListItem.h b/Engine/Plugins/Experimental/CommonUI/Source/CommonUI/Public/CommonNativeListItem.h
--- a/Engine/Plugins/Experimental/CommonUI/Source/CommonUI/Public/CommonNativeListItem.h
+++ b/Engine/Plugins/Experimental/CommonUI/Source/CommonUI/Public/CommonNativeListItem.h
@@ -2,7 +2,7 @@
 
 #pragma once
 
-#include "CoreMinimal.h"
+#include "Templates/SharedPointer.h"
 
 /** 
  * Base item class for any UMG ListViews based on native, non-UObject items.
@@ -87,4 +87,8 @@ public:
 private:
 	TSharedPtr<SListView<TSharedPtr<FMyCustomListItem>>> MyListView;
 }
-*/
\ No newline at end of file
+*/
+
+#if UE_ENABLE_INCLUDE_ORDER_DEPRECATED_IN_5_2
+#include "CoreMinimal.h"
+#endif
diff --git a/Engine/Plugins/Experimental/CommonUI/Source/CommonUI/Public/CommonNumericTextBlock.h b/Engine/Plugins/Experimental/CommonUI/Source/CommonUI/Public/CommonNumericTextBlock.h
--- a/Engine/Plugins/Experimental/CommonUI/Source/CommonUI/Public/CommonNumericTextBlock.h
+++ b/Engine/Plugins/Experimental/CommonUI/Source/CommonUI/Public/CommonNumericTextBlock.h
@@ -2,9 +2,7 @@
 
 #pragma once
 
-#include "CommonUserWidget.h"
 #include "CommonTextBlock.h"
-#include "Tickable.h"
 #include "CommonNumericTextBlock.generated.h"
 
 USTRUCT(BlueprintType)
@@ -219,3 +217,8 @@ private:
 		float Duration;
 	} SizeInterpolationState;
 };
+
+#if UE_ENABLE_INCLUDE_ORDER_DEPRECATED_IN_5_2
+#include "CommonUserWidget.h"
+#include "Tickable.h"
+#endif
diff --git a/Engine/Plugins/Experimental/CommonUI/Source/CommonUI/Public/CommonPoolableWidgetInterface.h b/Engine/Plugins/Experimental/CommonUI/Source/CommonUI/Public/CommonPoolableWidgetInterface.h
--- a/Engine/Plugins/Experimental/CommonUI/Source/CommonUI/Public/CommonPoolableWidgetInterface.h
+++ b/Engine/Plugins/Experimental/CommonUI/Source/CommonUI/Public/CommonPoolableWidgetInterface.h
@@ -2,9 +2,7 @@
 
 #pragma once
 
-#include "CoreMinimal.h"
 #include "UObject/Interface.h"
-#include "UObject/ScriptMacros.h"
 #include "CommonPoolableWidgetInterface.generated.h"
 
 UINTERFACE()
@@ -26,4 +24,9 @@ protected:
 
 	UFUNCTION(BlueprintNativeEvent, Category = "Common Poolable Widget")
 	void OnReleaseToPool();
-};
\ No newline at end of file
+};
+
+#if UE_ENABLE_INCLUDE_ORDER_DEPRECATED_IN_5_2
+#include "CoreMinimal.h"
+#include "UObject/ScriptMacros.h"
+#endif
diff --git a/Engine/Plugins/Experimental/CommonUI/Source/CommonUI/Public/CommonRichTextBlock.h b/Engine/Plugins/Experimental/CommonUI/Source/CommonUI/Public/CommonRichTextBlock.h
--- a/Engine/Plugins/Experimental/CommonUI/Source/CommonUI/Public/CommonRichTextBlock.h
+++ b/Engine/Plugins/Experimental/CommonUI/Source/CommonUI/Public/CommonRichTextBlock.h
@@ -3,10 +3,12 @@
 #pragma once
 
 #include "Components/RichTextBlock.h"
-#include "Styling/SlateBrush.h"
 
 #include "CommonRichTextBlock.generated.h"
 
+class ITextDecorator;
+enum class ETextTransformPolicy : uint8;
+
 class STextScroller;
 class UCommonTextStyle;
 class UCommonTextScrollStyle;
@@ -95,3 +97,7 @@ private:
 
 	TSharedPtr<STextScroller> MyTextScroller;
 };
+
+#if UE_ENABLE_INCLUDE_ORDER_DEPRECATED_IN_5_2
+#include "Styling/SlateBrush.h"
+#endif
diff --git a/Engine/Plugins/Experimental/CommonUI/Source/CommonUI/Public/CommonRotator.h b/Engine/Plugins/Experimental/CommonUI/Source/CommonUI/Public/CommonRotator.h
--- a/Engine/Plugins/Experimental/CommonUI/Source/CommonUI/Public/CommonRotator.h
+++ b/Engine/Plugins/Experimental/CommonUI/Source/CommonUI/Public/CommonRotator.h
@@ -2,16 +2,14 @@
 
 #pragma once
 
-#include "CommonUserWidget.h"
-#include "CommonUITypes.h"
 #include "CommonButtonBase.h"
-#include "CommonTextBlock.h"
-#include "CommonBorder.h"
 
-#include "Types/NavigationMetaData.h"
 
+#include "Input/NavigationReply.h"
 #include "CommonRotator.generated.h"
 
+class UCommonTextBlock;
+
 DECLARE_DYNAMIC_MULTICAST_DELEGATE_OneParam(FOnRotated, int32, Value);
 
 /**
@@ -86,3 +84,11 @@ protected:
 	/** The index of the current text item */
 	int32 SelectedIndex;
 };
+
+#if UE_ENABLE_INCLUDE_ORDER_DEPRECATED_IN_5_2
+#include "CommonBorder.h"
+#include "CommonTextBlock.h"
+#include "CommonUITypes.h"
+#include "CommonUserWidget.h"
+#include "Types/NavigationMetaData.h"
+#endif
diff --git a/Engine/Plugins/Experimental/CommonUI/Source/CommonUI/Public/CommonTabListWidgetBase.h b/Engine/Plugins/Experimental/CommonUI/Source/CommonUI/Public/CommonTabListWidgetBase.h
--- a/Engine/Plugins/Experimental/CommonUI/Source/CommonUI/Public/CommonTabListWidgetBase.h
+++ b/Engine/Plugins/Experimental/CommonUI/Source/CommonUI/Public/CommonTabListWidgetBase.h
@@ -3,10 +3,12 @@
 #pragma once
 
 #include "CommonUserWidget.h"
-#include "CommonUITypes.h"
 
+#include "Engine/DataTable.h"
 #include "CommonTabListWidgetBase.generated.h"
 
+class UCommonAnimatedSwitcher;
+
 class UCommonButtonBase;
 class UCommonButtonGroupBase;
 
@@ -231,3 +233,7 @@ private:
 	FUIActionBindingHandle NextTabActionHandle;
 	FUIActionBindingHandle PrevTabActionHandle;
 };
+
+#if UE_ENABLE_INCLUDE_ORDER_DEPRECATED_IN_5_2
+#include "CommonUITypes.h"
+#endif
diff --git a/Engine/Plugins/Experimental/CommonUI/Source/CommonUI/Public/CommonTextBlock.h b/Engine/Plugins/Experimental/CommonUI/Source/CommonUI/Public/CommonTextBlock.h
--- a/Engine/Plugins/Experimental/CommonUI/Source/CommonUI/Public/CommonTextBlock.h
+++ b/Engine/Plugins/Experimental/CommonUI/Source/CommonUI/Public/CommonTextBlock.h
@@ -2,11 +2,13 @@
 
 #pragma once
 
-#include "CommonUITypes.h"
-#include "Widgets/Layout/SScrollBox.h"
 #include "Components/TextBlock.h"
+#include "Widgets/Accessibility/SlateWidgetAccessibleTypes.h"
+#include "Widgets/SCompoundWidget.h"
 #include "CommonTextBlock.generated.h"
 
+struct FTextBlockStyle;
+
 class UCommonStyleSheet;
 
 /* 
@@ -259,4 +261,9 @@ private:
 	TSharedPtr<class STextScroller> TextScroller;
 
 	void ApplyFontSizeMultiplier() const;
-};
\ No newline at end of file
+};
+
+#if UE_ENABLE_INCLUDE_ORDER_DEPRECATED_IN_5_2
+#include "CommonUITypes.h"
+#include "Widgets/Layout/SScrollBox.h"
+#endif
diff --git a/Engine/Plugins/Experimental/CommonUI/Source/CommonUI/Public/CommonTileView.h b/Engine/Plugins/Experimental/CommonUI/Source/CommonUI/Public/CommonTileView.h
--- a/Engine/Plugins/Experimental/CommonUI/Source/CommonUI/Public/CommonTileView.h
+++ b/Engine/Plugins/Experimental/CommonUI/Source/CommonUI/Public/CommonTileView.h
@@ -5,6 +5,8 @@
 #include "Components/TileView.h"
 #include "CommonTileView.generated.h"
 
+class STableViewBase;
+
 /**
  * TileView specialized to navigate on focus for consoles & enable scrolling when not focused for touch
  */
@@ -19,4 +21,4 @@ public:
 protected:
 	virtual TSharedRef<STableViewBase> RebuildListWidget() override;
 	virtual UUserWidget& OnGenerateEntryWidgetInternal(UObject* Item, TSubclassOf<UUserWidget> DesiredEntryClass, const TSharedRef<STableViewBase>& OwnerTable) override;
-};
\ No newline at end of file
+};
diff --git a/Engine/Plugins/Experimental/CommonUI/Source/CommonUI/Public/CommonTreeView.h b/Engine/Plugins/Experimental/CommonUI/Source/CommonUI/Public/CommonTreeView.h
--- a/Engine/Plugins/Experimental/CommonUI/Source/CommonUI/Public/CommonTreeView.h
+++ b/Engine/Plugins/Experimental/CommonUI/Source/CommonUI/Public/CommonTreeView.h
@@ -5,6 +5,8 @@
 #include "Components/TreeView.h"
 #include "CommonTreeView.generated.h"
 
+class STableViewBase;
+
 //////////////////////////////////////////////////////////////////////////
 // SCommonTreeView
 //////////////////////////////////////////////////////////////////////////
@@ -95,4 +97,4 @@ public:
 protected:
 	virtual TSharedRef<STableViewBase> RebuildListWidget() override;
 	virtual UUserWidget& OnGenerateEntryWidgetInternal(UObject* Item, TSubclassOf<UUserWidget> DesiredEntryClass, const TSharedRef<STableViewBase>& OwnerTable) override;
-};
\ No newline at end of file
+};
diff --git a/Engine/Plugins/Experimental/CommonUI/Source/CommonUI/Public/CommonUIEditorSettings.h b/Engine/Plugins/Experimental/CommonUI/Source/CommonUI/Public/CommonUIEditorSettings.h
--- a/Engine/Plugins/Experimental/CommonUI/Source/CommonUI/Public/CommonUIEditorSettings.h
+++ b/Engine/Plugins/Experimental/CommonUI/Source/CommonUI/Public/CommonUIEditorSettings.h
@@ -2,16 +2,7 @@
 
 #pragma once
 
-#include "CoreMinimal.h"
-#include "UObject/Object.h"
-#include "UObject/SoftObjectPath.h"
-#include "Templates/SubclassOf.h"
-#include "InputCoreTypes.h"
-#include "Engine/StreamableManager.h"
-#include "Templates/SharedPointer.h"
 
-#include "ICommonUIModule.h"
-#include "CommonUITypes.h"
 #include "CommonTextBlock.h"
 #include "CommonButtonBase.h"
 #include "CommonBorder.h"
@@ -67,3 +58,15 @@ private:
 private:
 	bool bDefaultDataLoaded;
 };
+
+#if UE_ENABLE_INCLUDE_ORDER_DEPRECATED_IN_5_2
+#include "CommonUITypes.h"
+#include "CoreMinimal.h"
+#include "Engine/StreamableManager.h"
+#include "ICommonUIModule.h"
+#include "InputCoreTypes.h"
+#include "Templates/SharedPointer.h"
+#include "Templates/SubclassOf.h"
+#include "UObject/Object.h"
+#include "UObject/SoftObjectPath.h"
+#endif
diff --git a/Engine/Plugins/Experimental/CommonUI/Source/CommonUI/Public/CommonUILibrary.h b/Engine/Plugins/Experimental/CommonUI/Source/CommonUI/Public/CommonUILibrary.h
--- a/Engine/Plugins/Experimental/CommonUI/Source/CommonUI/Public/CommonUILibrary.h
+++ b/Engine/Plugins/Experimental/CommonUI/Source/CommonUI/Public/CommonUILibrary.h
@@ -2,12 +2,12 @@
 
 #pragma once
 
-#include "CoreMinimal.h"
-#include "UObject/ObjectMacros.h"
 #include "Kismet/BlueprintFunctionLibrary.h"
-#include "Components/Widget.h"
 #include "CommonUILibrary.generated.h"
 
+class UWidget;
+template <typename T> class TSubclassOf;
+
 UCLASS()
 class COMMONUI_API UCommonUILibrary : public UBlueprintFunctionLibrary
 {
@@ -21,3 +21,9 @@ public:
 	UFUNCTION(BlueprintCallable, Category="Common UI", meta=(DeterminesOutputType=Type))
 	static UWidget* FindParentWidgetOfType(UWidget* StartingWidget, TSubclassOf<UWidget> Type);
 };
+
+#if UE_ENABLE_INCLUDE_ORDER_DEPRECATED_IN_5_2
+#include "Components/Widget.h"
+#include "CoreMinimal.h"
+#include "UObject/ObjectMacros.h"
+#endif
diff --git a/Engine/Plugins/Experimental/CommonUI/Source/CommonUI/Public/CommonUIRichTextData.h b/Engine/Plugins/Experimental/CommonUI/Source/CommonUI/Public/CommonUIRichTextData.h
--- a/Engine/Plugins/Experimental/CommonUI/Source/CommonUI/Public/CommonUIRichTextData.h
+++ b/Engine/Plugins/Experimental/CommonUI/Source/CommonUI/Public/CommonUIRichTextData.h
@@ -2,8 +2,6 @@
 
 #pragma once
 
-#include "UObject/Object.h"
-#include "UObject/SoftObjectPtr.h"
 #include "Engine/DataTable.h"
 
 #include "CommonUIRichTextData.generated.h"
@@ -41,4 +39,9 @@ public:
 private:
 	UPROPERTY(EditDefaultsOnly, Category = "Inline Icons", meta = (RowType = "/Script/CommonUI.RichTextIconData"))
 	TObjectPtr<UDataTable> InlineIconSet;
-};
\ No newline at end of file
+};
+
+#if UE_ENABLE_INCLUDE_ORDER_DEPRECATED_IN_5_2
+#include "UObject/Object.h"
+#include "UObject/SoftObjectPtr.h"
+#endif
diff --git a/Engine/Plugins/Experimental/CommonUI/Source/CommonUI/Public/CommonUISettings.h b/Engine/Plugins/Experimental/CommonUI/Source/CommonUI/Public/CommonUISettings.h
--- a/Engine/Plugins/Experimental/CommonUI/Source/CommonUI/Public/CommonUISettings.h
+++ b/Engine/Plugins/Experimental/CommonUI/Source/CommonUI/Public/CommonUISettings.h
@@ -2,18 +2,21 @@
 
 #pragma once
 
-#include "CoreMinimal.h"
-#include "UObject/Object.h"
-#include "UObject/SoftObjectPath.h"
-#include "Templates/SharedPointer.h"
 
-#include "ICommonUIModule.h"
-#include "CommonUITypes.h"
 #include "CommonUIRichTextData.h"
-#include "GameplayTagContainer.h"
 #include "NativeGameplayTags.h"
+#include "Styling/SlateBrush.h"
+
+class UMaterialInterface;
 
 #if UE_ENABLE_INCLUDE_ORDER_DEPRECATED_IN_5_2
+#include "CommonUITypes.h"
+#include "CoreMinimal.h"
+#include "GameplayTagContainer.h"
+#include "ICommonUIModule.h"
+#include "Templates/SharedPointer.h"
+#include "UObject/Object.h"
+#include "UObject/SoftObjectPath.h"
 #include "CommonTextBlock.h"
 #endif
 
@@ -98,4 +101,4 @@ private:
 
 	UPROPERTY(Transient)
 	TObjectPtr<UCommonUIRichTextData> RichTextDataInstance;
-};
\ No newline at end of file
+};
diff --git a/Engine/Plugins/Experimental/CommonUI/Source/CommonUI/Public/CommonUISubsystemBase.h b/Engine/Plugins/Experimental/CommonUI/Source/CommonUI/Public/CommonUISubsystemBase.h
--- a/Engine/Plugins/Experimental/CommonUI/Source/CommonUI/Public/CommonUISubsystemBase.h
+++ b/Engine/Plugins/Experimental/CommonUI/Source/CommonUI/Public/CommonUISubsystemBase.h
@@ -3,11 +3,13 @@
 #pragma once
 
 #include "Subsystems/GameInstanceSubsystem.h"
-#include "Engine/DataTable.h"
-#include "Styling/SlateBrush.h"
 
 #include "CommonUISubsystemBase.generated.h"
 
+enum class ECommonInputType : uint8;
+struct FDataTableRowHandle;
+struct FSlateBrush;
+
 class IAnalyticsProviderET;
 class UWidget;
 class ULocalPlayer;
@@ -46,4 +48,9 @@ private:
 	void HandleInputMethodChanged(ECommonInputType bNewInputType);
 
 	TWeakPtr<class IAnalyticsProviderET> AnalyticProviderWeakPtr;
-};
\ No newline at end of file
+};
+
+#if UE_ENABLE_INCLUDE_ORDER_DEPRECATED_IN_5_2
+#include "Engine/DataTable.h"
+#include "Styling/SlateBrush.h"
+#endif
diff --git a/Engine/Plugins/Experimental/CommonUI/Source/CommonUI/Public/CommonUITypes.h b/Engine/Plugins/Experimental/CommonUI/Source/CommonUI/Public/CommonUITypes.h
--- a/Engine/Plugins/Experimental/CommonUI/Source/CommonUI/Public/CommonUITypes.h
+++ b/Engine/Plugins/Experimental/CommonUI/Source/CommonUI/Public/CommonUITypes.h
@@ -3,12 +3,14 @@
 #pragma once
 
 #include "Engine/DataTable.h"
-#include "Blueprint/UserWidget.h"
-#include "Misc/EnumRange.h"
-#include "Styling/SlateTypes.h"
-#include "CommonInputBaseTypes.h"
+#include "InputCoreTypes.h"
+#include "Styling/SlateBrush.h"
 #include "CommonUITypes.generated.h"
 
+class UUserWidget;
+enum class ECommonInputType : uint8;
+struct FScrollBoxStyle;
+
 class UCommonInputSubsystem;
 
 UENUM(BlueprintType)
@@ -175,3 +177,10 @@ public:
 
 DECLARE_DYNAMIC_DELEGATE_OneParam(FOnItemClicked, UUserWidget*, Widget);
 DECLARE_DYNAMIC_DELEGATE_TwoParams(FOnItemSelected, UUserWidget*, Widget, bool, Selected);
+
+#if UE_ENABLE_INCLUDE_ORDER_DEPRECATED_IN_5_2
+#include "Blueprint/UserWidget.h"
+#include "CommonInputBaseTypes.h"
+#include "Misc/EnumRange.h"
+#include "Styling/SlateTypes.h"
+#endif
diff --git a/Engine/Plugins/Experimental/CommonUI/Source/CommonUI/Public/CommonUIUtils.h b/Engine/Plugins/Experimental/CommonUI/Source/CommonUI/Public/CommonUIUtils.h
--- a/Engine/Plugins/Experimental/CommonUI/Source/CommonUI/Public/CommonUIUtils.h
+++ b/Engine/Plugins/Experimental/CommonUI/Source/CommonUI/Public/CommonUIUtils.h
@@ -2,8 +2,6 @@
 
 #pragma once
 
-#include "CoreMinimal.h"
-#include "Blueprint/UserWidget.h"
 #include "Blueprint/WidgetTree.h"
 
 class UWidget;
@@ -91,4 +89,9 @@ namespace CommonUIUtils
 
 	COMMONUI_API void ValidateBoundWidgetHierarchy(const UWidgetTree& WidgetTree, IWidgetCompilerLog& CompileLog, ECollisionPolicy CollisionPolicy, FName ParentWidgetName, TArray<FName>&& ChildNames);
 #endif
-};
\ No newline at end of file
+};
+
+#if UE_ENABLE_INCLUDE_ORDER_DEPRECATED_IN_5_2
+#include "Blueprint/UserWidget.h"
+#include "CoreMinimal.h"
+#endif
diff --git a/Engine/Plugins/Experimental/CommonUI/Source/CommonUI/Public/CommonUIVisibilitySubsystem.h b/Engine/Plugins/Experimental/CommonUI/Source/CommonUI/Public/CommonUIVisibilitySubsystem.h
--- a/Engine/Plugins/Experimental/CommonUI/Source/CommonUI/Public/CommonUIVisibilitySubsystem.h
+++ b/Engine/Plugins/Experimental/CommonUI/Source/CommonUI/Public/CommonUIVisibilitySubsystem.h
@@ -4,6 +4,7 @@
 
 #include "Subsystems/LocalPlayerSubsystem.h"
 #include "GameplayTagContainer.h"
+#include "UObject/WeakObjectPtr.h"
 #include "CommonUIVisibilitySubsystem.generated.h"
 
 class UWidget;
@@ -64,4 +65,4 @@ private:
 	static FGameplayTagContainer DebugTagsToEnable;
 	static FGameplayTagContainer DebugTagsToSuppress;
 #endif
-};
\ No newline at end of file
+};
diff --git a/Engine/Plugins/Experimental/CommonUI/Source/CommonUI/Public/CommonUserWidget.h b/Engine/Plugins/Experimental/CommonUI/Source/CommonUI/Public/CommonUserWidget.h
--- a/Engine/Plugins/Experimental/CommonUI/Source/CommonUI/Public/CommonUserWidget.h
+++ b/Engine/Plugins/Experimental/CommonUI/Source/CommonUI/Public/CommonUserWidget.h
@@ -2,7 +2,7 @@
 
 #pragma once
 
-#include "CommonUITypes.h"
+#include "Blueprint/UserWidget.h"
 #include "Input/UIActionBindingHandle.h"
 
 #include "CommonUserWidget.generated.h"
@@ -87,4 +87,8 @@ private:
 
 	TArray<FUIActionBindingHandle> ActionBindings;
 	TArray<TWeakObjectPtr<const UWidget>> ScrollRecipients;
-};
\ No newline at end of file
+};
+
+#if UE_ENABLE_INCLUDE_ORDER_DEPRECATED_IN_5_2
+#include "CommonUITypes.h"
+#endif
diff --git a/Engine/Plugins/Experimental/CommonUI/Source/CommonUI/Public/CommonVisibilitySwitcherSlot.h b/Engine/Plugins/Experimental/CommonUI/Source/CommonUI/Public/CommonVisibilitySwitcherSlot.h
--- a/Engine/Plugins/Experimental/CommonUI/Source/CommonUI/Public/CommonVisibilitySwitcherSlot.h
+++ b/Engine/Plugins/Experimental/CommonUI/Source/CommonUI/Public/CommonVisibilitySwitcherSlot.h
@@ -6,6 +6,8 @@
 
 #include "CommonVisibilitySwitcherSlot.generated.h"
 
+class SOverlay;
+
 class SBox;
 enum class ESlateVisibility : uint8;
 
diff --git a/Engine/Plugins/Experimental/CommonUI/Source/CommonUI/Public/CommonVisibilityWidgetBase.h b/Engine/Plugins/Experimental/CommonUI/Source/CommonUI/Public/CommonVisibilityWidgetBase.h
--- a/Engine/Plugins/Experimental/CommonUI/Source/CommonUI/Public/CommonVisibilityWidgetBase.h
+++ b/Engine/Plugins/Experimental/CommonUI/Source/CommonUI/Public/CommonVisibilityWidgetBase.h
@@ -2,11 +2,12 @@
 
 #pragma once
 
-#include "CoreMinimal.h"
-#include "CommonUITypes.h"
 #include "CommonBorder.h"
 #include "CommonVisibilityWidgetBase.generated.h"
 
+enum class ECommonInputType : uint8;
+enum class ESlateVisibility : uint8;
+
 /**
  * A container that controls visibility based on Input type and Platform
  *
@@ -49,3 +50,8 @@ protected:
 	UFUNCTION()
 	static const TArray<FName>& GetRegisteredPlatforms();
 };
+
+#if UE_ENABLE_INCLUDE_ORDER_DEPRECATED_IN_5_2
+#include "CommonUITypes.h"
+#include "CoreMinimal.h"
+#endif
diff --git a/Engine/Plugins/Experimental/CommonUI/Source/CommonUI/Public/CommonVisualAttachment.h b/Engine/Plugins/Experimental/CommonUI/Source/CommonUI/Public/CommonVisualAttachment.h
--- a/Engine/Plugins/Experimental/CommonUI/Source/CommonUI/Public/CommonVisualAttachment.h
+++ b/Engine/Plugins/Experimental/CommonUI/Source/CommonUI/Public/CommonVisualAttachment.h
@@ -2,9 +2,6 @@
 
 #pragma once
 
-#include "CoreMinimal.h"
-#include "UObject/ObjectMacros.h"
-#include "Widgets/SWidget.h"
 #include "Components/SizeBox.h"
 
 #include "CommonVisualAttachment.generated.h"
@@ -34,3 +31,9 @@ protected:
 private:
 	TSharedPtr<class SVisualAttachmentBox> MyAttachmentBox;
 };
+
+#if UE_ENABLE_INCLUDE_ORDER_DEPRECATED_IN_5_2
+#include "CoreMinimal.h"
+#include "UObject/ObjectMacros.h"
+#include "Widgets/SWidget.h"
+#endif
diff --git a/Engine/Plugins/Experimental/CommonUI/Source/CommonUI/Public/CommonWidgetCarousel.h b/Engine/Plugins/Experimental/CommonUI/Source/CommonUI/Public/CommonWidgetCarousel.h
--- a/Engine/Plugins/Experimental/CommonUI/Source/CommonUI/Public/CommonWidgetCarousel.h
+++ b/Engine/Plugins/Experimental/CommonUI/Source/CommonUI/Public/CommonWidgetCarousel.h
@@ -9,6 +9,9 @@ class UCommonWidgetCarousel;
 
 #include "CommonWidgetCarousel.generated.h"
 
+enum EHorizontalAlignment : int;
+enum EVerticalAlignment : int;
+
 DECLARE_DYNAMIC_MULTICAST_DELEGATE_TwoParams(FOnCurrentPageIndexChanged, UCommonWidgetCarousel*, CarouselWidget, int32, CurrentPageIndex);
 
 /**
diff --git a/Engine/Plugins/Experimental/CommonUI/Source/CommonUI/Public/CommonWidgetCarouselNavBar.h b/Engine/Plugins/Experimental/CommonUI/Source/CommonUI/Public/CommonWidgetCarouselNavBar.h
--- a/Engine/Plugins/Experimental/CommonUI/Source/CommonUI/Public/CommonWidgetCarouselNavBar.h
+++ b/Engine/Plugins/Experimental/CommonUI/Source/CommonUI/Public/CommonWidgetCarouselNavBar.h
@@ -2,13 +2,13 @@
 
 #pragma once
 
-#include "Components/PanelWidget.h"
 
 class UCommonWidgetCarousel;
 class UCommonButtonBase;
 class UCommonButtonGroupBase;
 class SHorizontalBox;
 
+#include "Components/Widget.h"
 #include "CommonWidgetCarouselNavBar.generated.h"
 
 /**
@@ -68,3 +68,7 @@ protected:
 	UPROPERTY()
 	TArray<TObjectPtr<UCommonButtonBase>> Buttons;
 };
+
+#if UE_ENABLE_INCLUDE_ORDER_DEPRECATED_IN_5_2
+#include "Components/PanelWidget.h"
+#endif
diff --git a/Engine/Plugins/Experimental/CommonUI/Source/CommonUI/Public/Input/CommonAnalogCursor.h b/Engine/Plugins/Experimental/CommonUI/Source/CommonUI/Public/Input/CommonAnalogCursor.h
--- a/Engine/Plugins/Experimental/CommonUI/Source/CommonUI/Public/Input/CommonAnalogCursor.h
+++ b/Engine/Plugins/Experimental/CommonUI/Source/CommonUI/Public/Input/CommonAnalogCursor.h
@@ -3,8 +3,8 @@
 #pragma once
 
 #include "Framework/Application/AnalogCursor.h"
-#include "Layout/Geometry.h"
 #include "InputCoreTypes.h"
+#include "Rendering/SlateRenderTransform.h"
 
 class UCommonUIActionRouterBase;
 class UCommonInputSubsystem;
@@ -124,3 +124,7 @@ private:
 	int32 ShoulderButtonStatus = EShoulderButtonFlags::None;
 #endif
 };
+
+#if UE_ENABLE_INCLUDE_ORDER_DEPRECATED_IN_5_2
+#include "Layout/Geometry.h"
+#endif
diff --git a/Engine/Plugins/Experimental/CommonUI/Source/CommonUI/Public/Input/CommonUIActionRouterBase.h b/Engine/Plugins/Experimental/CommonUI/Source/CommonUI/Public/Input/CommonUIActionRouterBase.h
--- a/Engine/Plugins/Experimental/CommonUI/Source/CommonUI/Public/Input/CommonUIActionRouterBase.h
+++ b/Engine/Plugins/Experimental/CommonUI/Source/CommonUI/Public/Input/CommonUIActionRouterBase.h
@@ -2,14 +2,19 @@
 
 #pragma once
 
+#include "CommonInputModeTypes.h"
 #include "Subsystems/LocalPlayerSubsystem.h"
-#include "CommonUIInputTypes.h"
-#include "CommonInputBaseTypes.h"
-#include "Misc/Passkey.h"
 #include "Containers/Ticker.h"
 
+#include "Engine/EngineBaseTypes.h"
+#include "Input/UIActionBindingHandle.h"
+#include "InputCoreTypes.h"
+#include "UObject/WeakObjectPtr.h"
 #include "CommonUIActionRouterBase.generated.h"
 
+class SWidget;
+struct FBindUIActionArgs;
+
 class AHUD;
 class FWeakWidgetPath;
 class FWidgetPath;
@@ -229,3 +234,9 @@ private:
 
 	TMap<TObjectPtr<UCommonInputActionDomain>, FActionDomainSortedRootList> ActionDomainRootNodes;
 };
+
+#if UE_ENABLE_INCLUDE_ORDER_DEPRECATED_IN_5_2
+#include "CommonInputBaseTypes.h"
+#include "CommonUIInputTypes.h"
+#include "Misc/Passkey.h"
+#endif
diff --git a/Engine/Plugins/Experimental/CommonUI/Source/CommonUI/Public/Input/CommonUIInputSettings.h b/Engine/Plugins/Experimental/CommonUI/Source/CommonUI/Public/Input/CommonUIInputSettings.h
--- a/Engine/Plugins/Experimental/CommonUI/Source/CommonUI/Public/Input/CommonUIInputSettings.h
+++ b/Engine/Plugins/Experimental/CommonUI/Source/CommonUI/Public/Input/CommonUIInputSettings.h
@@ -2,8 +2,7 @@
 
 #pragma once
 
-#include "UObject/Object.h"
-#include "Input/CommonUIInputTypes.h"
+#include "InputCoreTypes.h"
 #include "UITag.h"
 #include "CommonUIInputSettings.generated.h"
 
@@ -161,4 +160,9 @@ private:
 //	UPROPERTY(VisibleAnywhere, Category = InputSettings)
 //	int32 InputModePriorities[(uint8)ECommonInputMode::MAX];
 //#endif
-};
\ No newline at end of file
+};
+
+#if UE_ENABLE_INCLUDE_ORDER_DEPRECATED_IN_5_2
+#include "Input/CommonUIInputTypes.h"
+#include "UObject/Object.h"
+#endif
diff --git a/Engine/Plugins/Experimental/CommonUI/Source/CommonUI/Public/Input/CommonUIInputTypes.h b/Engine/Plugins/Experimental/CommonUI/Source/CommonUI/Public/Input/CommonUIInputTypes.h
--- a/Engine/Plugins/Experimental/CommonUI/Source/CommonUI/Public/Input/CommonUIInputTypes.h
+++ b/Engine/Plugins/Experimental/CommonUI/Source/CommonUI/Public/Input/CommonUIInputTypes.h
@@ -2,11 +2,10 @@
 
 #pragma once
 
-#include "UIActionBindingHandle.h"
+#include "Engine/DataTable.h"
 #include "UITag.h"
-#include "Misc/EnumRange.h"
 #include "CommonInputModeTypes.h"
-#include "InputCoreTypes.h"
+#include "Engine/EngineBaseTypes.h"
 
 struct COMMONUI_API FBindUIActionArgs
 {
@@ -69,3 +68,9 @@ struct COMMONUI_API FBindUIActionArgs
 	DECLARE_DELEGATE_OneParam(FOnHoldActionProgressed, float);
 	FOnHoldActionProgressed OnHoldActionProgressed;
 };
+
+#if UE_ENABLE_INCLUDE_ORDER_DEPRECATED_IN_5_2
+#include "InputCoreTypes.h"
+#include "Misc/EnumRange.h"
+#include "UIActionBindingHandle.h"
+#endif
diff --git a/Engine/Plugins/Experimental/CommonUI/Source/CommonUI/Public/Input/UIActionBinding.h b/Engine/Plugins/Experimental/CommonUI/Source/CommonUI/Public/Input/UIActionBinding.h
--- a/Engine/Plugins/Experimental/CommonUI/Source/CommonUI/Public/Input/UIActionBinding.h
+++ b/Engine/Plugins/Experimental/CommonUI/Source/CommonUI/Public/Input/UIActionBinding.h
@@ -2,9 +2,13 @@
 
 #pragma once
 
-#include "CoreMinimal.h"
-#include "Engine/EngineBaseTypes.h"
-#include "Input/CommonUIInputSettings.h"
+#include "Engine/DataTable.h"
+#include "Input/CommonUIInputSettings.h" // IWYU pragma: keep
+#include "Input/UIActionBindingHandle.h"
+
+enum EInputEvent : int;
+struct FKey;
+struct FUIActionKeyMapping;
 
 class FActionRouterBindingCollection;
 struct FBindUIActionArgs;
@@ -79,4 +83,9 @@ private:
 	static TMap<FKey, FUIActionBindingHandle> CurrentHoldActionKeys;
 
 	friend struct FUIActionBindingHandle;
-};
\ No newline at end of file
+};
+
+#if UE_ENABLE_INCLUDE_ORDER_DEPRECATED_IN_5_2
+#include "CoreMinimal.h"
+#include "Engine/EngineBaseTypes.h"
+#endif
diff --git a/Engine/Plugins/Experimental/CommonUI/Source/CommonUI/Public/Input/UIActionBindingHandle.h b/Engine/Plugins/Experimental/CommonUI/Source/CommonUI/Public/Input/UIActionBindingHandle.h
--- a/Engine/Plugins/Experimental/CommonUI/Source/CommonUI/Public/Input/UIActionBindingHandle.h
+++ b/Engine/Plugins/Experimental/CommonUI/Source/CommonUI/Public/Input/UIActionBindingHandle.h
@@ -2,8 +2,11 @@
 
 #pragma once
 
-#include "CoreMinimal.h"
-#include "Engine/EngineBaseTypes.h"
+#include "Containers/UnrealString.h"
+#include "Misc/Optional.h"
+
+enum class EMouseCaptureMode : uint8;
+struct FScriptContainerElement;
 
 class UWidget;
 enum class ECommonInputMode : uint8;
@@ -90,4 +93,9 @@ private:
 	ECommonInputMode InputMode;
 	EMouseCaptureMode MouseCaptureMode;
 	bool bHideCursorDuringViewportCapture = true;
-};
\ No newline at end of file
+};
+
+#if UE_ENABLE_INCLUDE_ORDER_DEPRECATED_IN_5_2
+#include "CoreMinimal.h"
+#include "Engine/EngineBaseTypes.h"
+#endif
diff --git a/Engine/Plugins/Experimental/CommonUI/Source/CommonUI/Public/SAnalogSlider.h b/Engine/Plugins/Experimental/CommonUI/Source/CommonUI/Public/SAnalogSlider.h
--- a/Engine/Plugins/Experimental/CommonUI/Source/CommonUI/Public/SAnalogSlider.h
+++ b/Engine/Plugins/Experimental/CommonUI/Source/CommonUI/Public/SAnalogSlider.h
@@ -4,6 +4,8 @@
 
 #include "Widgets/Input/SSlider.h"
 
+struct FAnalogInputEvent;
+
 /**
  * A Slate slider control is a linear scale and draggable handle.
  */
diff --git a/Engine/Plugins/Experimental/CommonUI/Source/CommonUI/Public/SCommonButtonTableRow.h b/Engine/Plugins/Experimental/CommonUI/Source/CommonUI/Public/SCommonButtonTableRow.h
--- a/Engine/Plugins/Experimental/CommonUI/Source/CommonUI/Public/SCommonButtonTableRow.h
+++ b/Engine/Plugins/Experimental/CommonUI/Source/CommonUI/Public/SCommonButtonTableRow.h
@@ -5,6 +5,8 @@
 #include "Slate/SObjectTableRow.h"
 #include "CommonButtonBase.h"
 
+class STableViewBase;
+
 /** 
  * A CommonUI version of the object table row that is aware of UCommonButtonBase.
  * Instead of bothering with handling mouse events directly, we rely on the entry being a button itself and respond to events from it.
diff --git a/Engine/Plugins/Experimental/CommonUI/Source/CommonUI/Public/SCommonHierarchicalScrollBox.h b/Engine/Plugins/Experimental/CommonUI/Source/CommonUI/Public/SCommonHierarchicalScrollBox.h
--- a/Engine/Plugins/Experimental/CommonUI/Source/CommonUI/Public/SCommonHierarchicalScrollBox.h
+++ b/Engine/Plugins/Experimental/CommonUI/Source/CommonUI/Public/SCommonHierarchicalScrollBox.h
@@ -2,7 +2,6 @@
 
 #pragma once
 
-#include "CoreMinimal.h"
 #include "Widgets/Layout/SScrollBox.h"
 
 /** SScrollBox can scroll through an arbitrary number of widgets. */
@@ -20,3 +19,7 @@ public:
 private:
 	void AppendFocusableChildren(TArray<TSharedRef<SWidget>>& OutChildren, TSharedRef<SWidget> Owner);
 };
+
+#if UE_ENABLE_INCLUDE_ORDER_DEPRECATED_IN_5_2
+#include "CoreMinimal.h"
+#endif
diff --git a/Engine/Plugins/Experimental/CommonUI/Source/CommonUI/Public/SVisualAttachmentBox.h b/Engine/Plugins/Experimental/CommonUI/Source/CommonUI/Public/SVisualAttachmentBox.h
--- a/Engine/Plugins/Experimental/CommonUI/Source/CommonUI/Public/SVisualAttachmentBox.h
+++ b/Engine/Plugins/Experimental/CommonUI/Source/CommonUI/Public/SVisualAttachmentBox.h
@@ -2,7 +2,6 @@
 
 #pragma once
 
-#include "CoreMinimal.h"
 #include "Widgets/Layout/SBox.h"
 
 /**
@@ -23,4 +22,8 @@ private:
 	mutable FVector2D InnerDesiredSize;
 
 	FVector2D ContentAnchor = FVector2D::ZeroVector;
-};
\ No newline at end of file
+};
+
+#if UE_ENABLE_INCLUDE_ORDER_DEPRECATED_IN_5_2
+#include "CoreMinimal.h"
+#endif
diff --git a/Engine/Plugins/Experimental/CommonUI/Source/CommonUI/Public/WidgetFactory.h b/Engine/Plugins/Experimental/CommonUI/Source/CommonUI/Public/WidgetFactory.h
--- a/Engine/Plugins/Experimental/CommonUI/Source/CommonUI/Public/WidgetFactory.h
+++ b/Engine/Plugins/Experimental/CommonUI/Source/CommonUI/Public/WidgetFactory.h
@@ -2,11 +2,11 @@
 
 #pragma once
 
-#include "Blueprint/UserWidget.h"
-#include "Widgets/Views/STableRow.h"
-#include "Slate/SObjectTableRow.h"
-#include "Widgets/Views/STableViewBase.h"
 #include "CommonPoolableWidgetInterface.h"
+#include "Slate/SObjectWidget.h"
+
+class ITableRow;
+class STableViewBase;
 
 #define ENABLE_WIDGET_FACTORY_POOLING 1
 
@@ -280,4 +280,11 @@ private:
 	bool	bTrackActiveWidgets;
 };
 
-using FUserWidgetFactory = TWidgetFactory<UUserWidget>;
\ No newline at end of file
+using FUserWidgetFactory = TWidgetFactory<UUserWidget>;
+
+#if UE_ENABLE_INCLUDE_ORDER_DEPRECATED_IN_5_2
+#include "Blueprint/UserWidget.h"
+#include "Slate/SObjectTableRow.h"
+#include "Widgets/Views/STableRow.h"
+#include "Widgets/Views/STableViewBase.h"
+#endif
diff --git a/Engine/Plugins/Experimental/CommonUI/Source/CommonUI/Public/Widgets/CommonActivatableWidgetContainer.h b/Engine/Plugins/Experimental/CommonUI/Source/CommonUI/Public/Widgets/CommonActivatableWidgetContainer.h
--- a/Engine/Plugins/Experimental/CommonUI/Source/CommonUI/Public/Widgets/CommonActivatableWidgetContainer.h
+++ b/Engine/Plugins/Experimental/CommonUI/Source/CommonUI/Public/Widgets/CommonActivatableWidgetContainer.h
@@ -2,11 +2,13 @@
 
 #pragma once
 
-#include "Components/Widget.h"
-#include "Slate/SCommonAnimatedSwitcher.h"
 #include "Blueprint/UserWidgetPool.h"
 #include "CommonActivatableWidgetContainer.generated.h"
 
+class SCommonAnimatedSwitcher;
+enum class ECommonSwitcherTransition : uint8;
+enum class ETransitionCurve : uint8;
+
 class UCommonActivatableWidget;
 class SOverlay;
 class SSpacer;
@@ -216,4 +218,9 @@ class COMMONUI_API UCommonActivatableWidgetQueue : public UCommonActivatableWidg
 
 protected:
 	virtual void OnWidgetAddedToList(UCommonActivatableWidget& AddedWidget) override;
-};
\ No newline at end of file
+};
+
+#if UE_ENABLE_INCLUDE_ORDER_DEPRECATED_IN_5_2
+#include "Components/Widget.h"
+#include "Slate/SCommonAnimatedSwitcher.h"
+#endif
diff --git a/Engine/Plugins/Experimental/CommonUI/Source/CommonUIEditor/Private/CommonAssetTypeActions_GenericInputActionDataTable.h b/Engine/Plugins/Experimental/CommonUI/Source/CommonUIEditor/Private/CommonAssetTypeActions_GenericInputActionDataTable.h
--- a/Engine/Plugins/Experimental/CommonUI/Source/CommonUIEditor/Private/CommonAssetTypeActions_GenericInputActionDataTable.h
+++ b/Engine/Plugins/Experimental/CommonUI/Source/CommonUIEditor/Private/CommonAssetTypeActions_GenericInputActionDataTable.h
@@ -2,8 +2,6 @@
 
 #pragma once
 
-#include "CoreMinimal.h"
-#include "Toolkits/IToolkitHost.h"
 #include "AssetTypeActions_CSVAssetBase.h"
 
 class UCommonGenericInputActionDataTable;
@@ -25,3 +23,8 @@ private:
 	/** Called via context menu to open data within the data table editor */
 	void OpenInDataTableEditor(TArray< TWeakObjectPtr<UCommonGenericInputActionDataTable> > Objects);
 };
+
+#if UE_ENABLE_INCLUDE_ORDER_DEPRECATED_IN_5_2
+#include "CoreMinimal.h"
+#include "Toolkits/IToolkitHost.h"
+#endif
diff --git a/Engine/Plugins/Experimental/CommonUI/Source/CommonUIEditor/Private/CommonVideoPlayerCustomization.h b/Engine/Plugins/Experimental/CommonUI/Source/CommonUIEditor/Private/CommonVideoPlayerCustomization.h
--- a/Engine/Plugins/Experimental/CommonUI/Source/CommonUIEditor/Private/CommonVideoPlayerCustomization.h
+++ b/Engine/Plugins/Experimental/CommonUI/Source/CommonUIEditor/Private/CommonVideoPlayerCustomization.h
@@ -4,8 +4,8 @@
 
 #include "IDetailCustomization.h"
 #include "ITransportControl.h"
-#include "PropertyEditorModule.h"
-#include "Types/SlateEnums.h"
+
+template <typename OptionalType> struct TOptional;
 
 class IPropertyHandle;
 class UCommonVideoPlayer;
@@ -38,3 +38,8 @@ private:
 
 	TWeakObjectPtr<UCommonVideoPlayer> VideoPlayer;
 };
+
+#if UE_ENABLE_INCLUDE_ORDER_DEPRECATED_IN_5_2
+#include "PropertyEditorModule.h"
+#include "Types/SlateEnums.h"
+#endif
diff --git a/Engine/Plugins/Experimental/CommonUI/Source/CommonUIEditor/Public/CommonGenericInputActionDataTableFactory.h b/Engine/Plugins/Experimental/CommonUI/Source/CommonUIEditor/Public/CommonGenericInputActionDataTableFactory.h
--- a/Engine/Plugins/Experimental/CommonUI/Source/CommonUIEditor/Public/CommonGenericInputActionDataTableFactory.h
+++ b/Engine/Plugins/Experimental/CommonUI/Source/CommonUIEditor/Public/CommonGenericInputActionDataTableFactory.h
@@ -2,8 +2,6 @@
 
 #pragma once
 
-#include "CoreMinimal.h"
-#include "UObject/ObjectMacros.h"
 #include "Factories/Factory.h"
 #include "CommonGenericInputActionDataTableFactory.generated.h"
 
@@ -17,3 +15,8 @@ class UCommonGenericInputActionDataTableFactory : public UFactory
 // End of UFactory interface
 
 };
+
+#if UE_ENABLE_INCLUDE_ORDER_DEPRECATED_IN_5_2
+#include "CoreMinimal.h"
+#include "UObject/ObjectMacros.h"
+#endif
diff --git a/Engine/Plugins/Experimental/CommonUI/Source/CommonUIEditor/Public/CommonUIEditorModule.h b/Engine/Plugins/Experimental/CommonUI/Source/CommonUIEditor/Public/CommonUIEditorModule.h
--- a/Engine/Plugins/Experimental/CommonUI/Source/CommonUIEditor/Public/CommonUIEditorModule.h
+++ b/Engine/Plugins/Experimental/CommonUI/Source/CommonUIEditor/Public/CommonUIEditorModule.h
@@ -3,7 +3,6 @@
 #pragma once
 
 #include "AssetTypeActions_Base.h"
-#include "Modules/ModuleInterface.h"
 
 class FCommonUIEditorModule
 	: public IModuleInterface
@@ -16,4 +15,8 @@ private:
 
 	/** Asset type actions */
 	TArray<TSharedPtr<class FAssetTypeActions_Base>> ItemDataAssetTypeActions;
-};
\ No newline at end of file
+};
+
+#if UE_ENABLE_INCLUDE_ORDER_DEPRECATED_IN_5_2
+#include "Modules/ModuleInterface.h"
+#endif
