Add CVar to log when input type changes to CommonUI
This time the CVar will be a permanent addition to CommonUI.

diff --git a/Engine/Plugins/Experimental/CommonUI/Source/CommonInput/Private/CommonInputSubsystem.cpp b/Engine/Plugins/Experimental/CommonUI/Source/CommonInput/Private/CommonInputSubsystem.cpp
--- a/Engine/Plugins/Experimental/CommonUI/Source/CommonInput/Private/CommonInputSubsystem.cpp
+++ b/Engine/Plugins/Experimental/CommonUI/Source/CommonInput/Private/CommonInputSubsystem.cpp
@@ -3,19 +3,19 @@
 #include "CommonInputSubsystem.h"
 #include "CommonInputPrivate.h"
 #include "CommonInputTypeEnum.h"
-#include "Framework/Application/SlateApplication.h"
 
+#include "Engine/Engine.h"
+#include "Engine/GameViewportClient.h"
+#include "Engine/LocalPlayer.h"
 #include "Engine/PlatformSettingsManager.h"
+#include "Engine/World.h"
 #include "Framework/Application/IInputProcessor.h"
 #include "Framework/Application/SlateApplication.h"
 #include "Framework/Application/SlateUser.h"
-#include "Engine/World.h"
-#include "Engine/GameViewportClient.h"
+#include "HAL/PlatformStackWalk.h"
 #include "Widgets/SViewport.h"
 #include "CommonInputSettings.h"
 #include "ICommonInputModule.h"
-#include "Engine/LocalPlayer.h"
-#include "Engine/Engine.h"
 
 #include UE_INLINE_GENERATED_CPP_BY_NAME(CommonInputSubsystem)
 
@@ -23,6 +23,15 @@
 #include "Settings/LevelEditorPlaySettings.h"
 #endif
 
+#if !UE_BUILD_SHIPPING
+static int32 bDumpInputTypeChangeCallstack = 0;
+static FAutoConsoleVariableRef CVarDumpInputTypeChangeCallstack(
+	TEXT("CommonUI.bDumpInputTypeChangeCallstack"),
+	bDumpInputTypeChangeCallstack,
+	TEXT("Dump callstack when input type changes."));
+#endif // !UE_BUILD_SHIPPING
+
+
 FPlatformInputSupportOverrideDelegate UCommonInputSubsystem::OnPlatformInputSupportOverride;
 
 /**
@@ -517,6 +526,19 @@ void UCommonInputSubsystem::SetCurrentInputType(ECommonInputType NewInputType)
 
 			if (LockedInput != CurrentInputType)
 			{
+#if !UE_BUILD_SHIPPING
+				if (bDumpInputTypeChangeCallstack)
+				{
+					const uint32 DumpCallstackSize = 65535;
+					ANSICHAR DumpCallstack[DumpCallstackSize] = { 0 };
+					FString ScriptStack = FFrame::GetScriptCallstack(true /* bReturnEmpty */);
+					FPlatformStackWalk::StackWalkAndDump(DumpCallstack, DumpCallstackSize, 0);
+					UE_LOG(LogCommonInput, Log, TEXT("--- Input Changing Callstack ---"));
+					UE_LOG(LogCommonInput, Log, TEXT("Script Stack:\n%s"), *ScriptStack);
+					UE_LOG(LogCommonInput, Log, TEXT("Callstack:\n%s"), ANSI_TO_TCHAR(DumpCallstack));
+				}
+#endif // !UE_BUILD_SHIPPING
+				
 				CurrentInputType = LockedInput;
 
 				FSlateApplication& SlateApplication = FSlateApplication::Get();
