[CommonUI]
* WidgetFactory - Removed second template parameter with default argument that is used to prevent class from being instantiated if argument to first parameter is not derived from UUserWidget. It is more readable and cleaner to use a static_assert inside the class instead. It makes it possible to forward declare WidgetFactory and T in WidgetFactory<T>

diff --git a/Engine/Plugins/Experimental/CommonUI/Source/CommonUI/Public/WidgetFactory.h b/Engine/Plugins/Experimental/CommonUI/Source/CommonUI/Public/WidgetFactory.h
--- a/Engine/Plugins/Experimental/CommonUI/Source/CommonUI/Public/WidgetFactory.h
+++ b/Engine/Plugins/Experimental/CommonUI/Source/CommonUI/Public/WidgetFactory.h
@@ -10,9 +10,10 @@ class STableViewBase;
 
 #define ENABLE_WIDGET_FACTORY_POOLING 1
 
-template <class WidgetType, class = typename TEnableIf<TIsDerivedFrom<WidgetType, UUserWidget>::IsDerived, WidgetType>::Type>
+template <class WidgetType>
 class TWidgetFactory : public FGCObject
 {
+	static_assert(TIsDerivedFrom<WidgetType, UUserWidget>::IsDerived, "Can only use TWidgetFactory with types deriving from UUserWidget");
 private:
 	FORCEINLINE UGameInstance* OwningObjectAsGameInstance() const { return Cast<UGameInstance>(OuterGetter()); }
 	FORCEINLINE UWorld* OwningObjectAsWorld() const { return Cast<UWorld>(OuterGetter()); }
