Fixed an issue where we would center the mouse due to a previous action domain root node config being applied even though we deactivated a non-action domain root node.
#rnx
#tests: Ran with various cases of UI scenarios in PIE.

diff --git a/Engine/Plugins/Experimental/CommonUI/Source/CommonUI/Private/Input/CommonUIActionRouterBase.cpp b/Engine/Plugins/Experimental/CommonUI/Source/CommonUI/Private/Input/CommonUIActionRouterBase.cpp
--- a/Engine/Plugins/Experimental/CommonUI/Source/CommonUI/Private/Input/CommonUIActionRouterBase.cpp
+++ b/Engine/Plugins/Experimental/CommonUI/Source/CommonUI/Private/Input/CommonUIActionRouterBase.cpp
@@ -608,7 +608,21 @@ void UCommonUIActionRouterBase::HandleRootNodeActivated(TWeakPtr<FActivatableTre
 {
 	FActivatableTreeRootRef ActivatedRoot = WeakActivatedRoot.Pin().ToSharedRef();
 	UCommonActivatableWidget* NodeWidget = ActivatedRoot->GetWidget();
-	if (UCommonInputActionDomain* WidgetActionDomain = NodeWidget ? NodeWidget->GetCalculatedActionDomain() : nullptr)
+
+	if (RootNodes.Contains(ActivatedRoot) && ActivatedRoot->GetLastPaintLayer() > 0)
+	{
+		const int32 CurrentRootLayer = ActiveRootNode ? ActiveRootNode->GetLastPaintLayer() : INDEX_NONE;
+		if (ActivatedRoot->GetLastPaintLayer() > CurrentRootLayer)
+		{
+			// Ensure we have a local player so the action router local player subsystem will handle root change
+			const ULocalPlayer* LocalPlayer = GetLocalPlayer();
+			if (LocalPlayer && LocalPlayer->ViewportClient)
+			{
+				SetActiveRoot(ActivatedRoot);
+			}
+		}
+	}
+	else if (UCommonInputActionDomain* WidgetActionDomain = NodeWidget ? NodeWidget->GetCalculatedActionDomain() : nullptr)
 	{
 		const FActionDomainSortedRootList* ActionDomainRootList = ActionDomainRootNodes.Find(WidgetActionDomain);
 		if (ActionDomainRootList && ensure(ActionDomainRootList->Contains(ActivatedRoot)))
@@ -619,9 +633,9 @@ void UCommonUIActionRouterBase::HandleRootNodeActivated(TWeakPtr<FActivatableTre
 			UCommonInputSubsystem& CommonInputSubsystem = GetInputSubsystem();
 			if (UCommonInputActionDomainTable* ActionDomainTable = CommonInputSubsystem.GetActionDomainTable())
 			{
-				// We find the first root node that isn't the activated root and bail early. 
+				// We find the first root node that is receiving input and bail early. 
 				// The action domains and root lists are sorted so we should end up with a node with a higher paint layer
-				// or we just stop at the activated node and make sure it's leaf node is updated.
+				// or we just stop at the activated root node and make sure it's leaf node is updated.
 				for (const UCommonInputActionDomain* ActionDomain : ActionDomainTable->ActionDomains)
 				{
 					if (FActionDomainSortedRootList* SortedRootList = ActionDomainRootNodes.Find(ActionDomain))
@@ -630,7 +644,7 @@ void UCommonUIActionRouterBase::HandleRootNodeActivated(TWeakPtr<FActivatableTre
 						{
 							if (RootNode->IsReceivingInput() && RootNode->DoesWidgetSupportActivationFocus())
 							{
-								if (RootNode == ActivatedRoot)
+								if (RootNode == ActivatedRoot && !ActiveRootNode.IsValid())
 								{
 									// This will allow us to update the leaf node config on the activated root.
 									ActivatedRoot->UpdateLeafNode();
@@ -643,46 +657,33 @@ void UCommonUIActionRouterBase::HandleRootNodeActivated(TWeakPtr<FActivatableTre
 				}
 			}
 		}
-
-		return;
-	}
-
-	if (RootNodes.Contains(ActivatedRoot))
-	{
-		if (ActivatedRoot->GetLastPaintLayer() > 0)
-		{
-			const int32 CurrentRootLayer = ActiveRootNode ? ActiveRootNode->GetLastPaintLayer() : INDEX_NONE;
-			if (ActivatedRoot->GetLastPaintLayer() > CurrentRootLayer)
-			{
-				// Ensure we have a local player so the action router local player subsystem will handle root change
-				const ULocalPlayer* LocalPlayer = GetLocalPlayer();
-				if (LocalPlayer && LocalPlayer->ViewportClient)
-				{
-					SetActiveRoot(ActivatedRoot);
-				}
-			}
-		}
 	}
 }
 
 void UCommonUIActionRouterBase::HandleRootNodeDeactivated(TWeakPtr<FActivatableTreeRoot> WeakDeactivatedRoot)
 {
 	FActivatableTreeRootPtr DeactivatedRoot = WeakDeactivatedRoot.Pin();
-	UCommonActivatableWidget* NodeWidget = DeactivatedRoot ? DeactivatedRoot->GetWidget() : nullptr;
-	UCommonInputActionDomain* ActionDomain = NodeWidget ? NodeWidget->GetCalculatedActionDomain() : nullptr;
-	
-	if (ActionDomain || (ActiveRootNode && ActiveRootNode == DeactivatedRoot))
+	if (ActiveRootNode.IsValid() && ActiveRootNode == DeactivatedRoot)
 	{
-		// Action domain root nodes will need to recalibrate during deactivation since we're not relying on 
-		// timing to fix up the input config as RootNodes does.
-		RefreshActionDomainLeafNodeConfig();
-
 		// Reset the active root widget - we'll re-establish it on the next tick
-		if (ActiveRootNode.IsValid())
+		SetActiveRoot(nullptr);
+	}
+
+	bool bActivatedRootNodeExists = false;
+	for (const FActivatableTreeRootRef& Root : RootNodes)
+	{
+		if (Root->IsWidgetActivated())
 		{
-			SetActiveRoot(nullptr);
+			bActivatedRootNodeExists = true;
+			break;
 		}
 	}
+
+	// In the case that all root nodes are not activated we need to re-establish input for the highest paint layer node in action domain nodes.
+	if (!bActivatedRootNodeExists)
+	{
+		RefreshActionDomainLeafNodeConfig();
+	}
 }
 
 void UCommonUIActionRouterBase::HandleLeafmostActiveNodeChanged()
