UCommonTabListWidgetBase:
- When removing a tab, decrement subsequent tabs' index.

diff --git a/Engine/Plugins/Experimental/CommonUI/Source/CommonUI/Private/CommonTabListWidgetBase.cpp b/Engine/Plugins/Experimental/CommonUI/Source/CommonUI/Private/CommonTabListWidgetBase.cpp
--- a/Engine/Plugins/Experimental/CommonUI/Source/CommonUI/Private/CommonTabListWidgetBase.cpp
+++ b/Engine/Plugins/Experimental/CommonUI/Source/CommonUI/Private/CommonTabListWidgetBase.cpp
@@ -126,23 +126,26 @@ bool UCommonTabListWidgetBase::RegisterTab(FName TabNameID, TSubclassOf<UCommonB
 
 bool UCommonTabListWidgetBase::RemoveTab(FName TabNameID)
 {
-	FCommonRegisteredTabInfo* const TabInfo = RegisteredTabsByID.Find(TabNameID);
+	const FCommonRegisteredTabInfo* TabInfo = RegisteredTabsByID.Find(TabNameID);
+
 	if (!TabInfo)
 	{
 		return false;
 	}
 
-	UCommonButtonBase* const TabButton = TabInfo->TabButton;
-	if (TabButton)
+	if (TabInfo->TabIndex >= 0)
 	{
-		TabButtonGroup->RemoveWidget(TabButton);
-		TabButton->RemoveFromParent();
+		for (TPair<FName, FCommonRegisteredTabInfo>& RegisteredTabByID : RegisteredTabsByID)
+		{
+			if (RegisteredTabByID.Value.TabIndex > TabInfo->TabIndex)
+			{
+				// Decrement this tab's index as we have removed a tab before it.
+				RegisteredTabByID.Value.TabIndex--;
+			}
+		}
 	}
-	RegisteredTabsByID.Remove(TabNameID);
 
-	// Callbacks
-	HandleTabRemoval(TabNameID, TabButton);
-	OnTabButtonRemoval.Broadcast(TabNameID, TabButton);
+	RemoveTab_Internal(TabNameID, *TabInfo);
 
 	return true;
 }
@@ -151,7 +154,7 @@ void UCommonTabListWidgetBase::RemoveAllTabs()
 {
 	for (TMap<FName, FCommonRegisteredTabInfo>::TIterator Iter(RegisteredTabsByID); Iter; ++Iter)
 	{
-		RemoveTab(Iter->Key);
+		RemoveTab_Internal(Iter->Key, Iter->Value);
 	}
 }
 
@@ -486,3 +489,19 @@ void UCommonTabListWidgetBase::RebuildTabList()
 	OnTabListRebuilt.Broadcast();
 }
 
+void UCommonTabListWidgetBase::RemoveTab_Internal(const FName TabNameID, const FCommonRegisteredTabInfo& TabInfo)
+{
+	UCommonButtonBase* const TabButton = TabInfo.TabButton;
+
+	if (TabButton)
+	{
+		TabButtonGroup->RemoveWidget(TabButton);
+		TabButton->RemoveFromParent();
+	}
+
+	RegisteredTabsByID.Remove(TabNameID);
+
+	// Callbacks
+	HandleTabRemoval(TabNameID, TabButton);
+	OnTabButtonRemoval.Broadcast(TabNameID, TabButton);
+}
diff --git a/Engine/Plugins/Experimental/CommonUI/Source/CommonUI/Public/CommonTabListWidgetBase.h b/Engine/Plugins/Experimental/CommonUI/Source/CommonUI/Public/CommonTabListWidgetBase.h
--- a/Engine/Plugins/Experimental/CommonUI/Source/CommonUI/Public/CommonTabListWidgetBase.h
+++ b/Engine/Plugins/Experimental/CommonUI/Source/CommonUI/Public/CommonTabListWidgetBase.h
@@ -220,6 +220,8 @@ private:
 	bool DeferredRebuildTabList(float DeltaTime);
 	void RebuildTabList();
 
+	void RemoveTab_Internal(const FName TabNameID, const FCommonRegisteredTabInfo& TabInfo);
+
 	/** Info about each of the currently registered tabs organized by a given registration name ID */
 	UPROPERTY(Transient)
 	TMap<FName, FCommonRegisteredTabInfo> RegisteredTabsByID;
