Fix CommonUI uses cached geometry that might be zero sized for gamepad focus.

diff --git a/Engine/Plugins/Experimental/CommonUI/Source/CommonUI/Private/Input/CommonAnalogCursor.cpp b/Engine/Plugins/Experimental/CommonUI/Source/CommonUI/Private/Input/CommonAnalogCursor.cpp
--- a/Engine/Plugins/Experimental/CommonUI/Source/CommonUI/Private/Input/CommonAnalogCursor.cpp
+++ b/Engine/Plugins/Experimental/CommonUI/Source/CommonUI/Private/Input/CommonAnalogCursor.cpp
@@ -100,7 +100,7 @@ void FCommonAnalogCursor::Tick(const float DeltaTime, FSlateApplication& SlateAp
 				}
 			}
 
-			// We want to update the cursor position when focus changes or the focused widget moves at all
+			// We want to try to update the cursor position when focus changes or the focused widget moves at all
 			if (CursorTarget != PinnedLastCursorTarget || (CursorTarget && CursorTarget->GetCachedGeometry().GetAccumulatedRenderTransform() != LastCursorTargetTransform))
 			{
 #if !UE_BUILD_SHIPPING
@@ -136,12 +136,13 @@ void FCommonAnalogCursor::Tick(const float DeltaTime, FSlateApplication& SlateAp
 					}
 					else
 					{
-						TargetGeometry = CursorTarget->GetCachedGeometry();
+						TargetGeometry = CursorTarget->GetTickSpaceGeometry();
 					}
-					
-					LastCursorTargetTransform = TargetGeometry.GetAccumulatedRenderTransform();
-					if (TargetGeometry.GetLocalSize().SizeSquared() > SMALL_NUMBER)
+
+					if (TargetGeometry.GetLocalSize().X > UE_SMALL_NUMBER && TargetGeometry.GetLocalSize().Y > UE_SMALL_NUMBER)
 					{
+						LastCursorTargetTransform = TargetGeometry.GetAccumulatedRenderTransform();
+
 						bHasValidCursorTarget = true;
 						
 						const FVector2D AbsoluteWidgetCenter = TargetGeometry.GetAbsolutePositionAtCoordinates(FVector2D(0.5f, 0.5f));
