Fixed CommonTextBlock size reset on mobile (MobileFontSizeMultiplier was ignored on widget invalidation).

diff --git a/Engine/Plugins/Experimental/CommonUI/Source/CommonUI/Private/CommonTextBlock.cpp b/Engine/Plugins/Experimental/CommonUI/Source/CommonUI/Private/CommonTextBlock.cpp
--- a/Engine/Plugins/Experimental/CommonUI/Source/CommonUI/Private/CommonTextBlock.cpp
+++ b/Engine/Plugins/Experimental/CommonUI/Source/CommonUI/Private/CommonTextBlock.cpp
@@ -611,15 +611,18 @@ void UCommonTextBlock::SynchronizeProperties()
 	}
 }
 
-void UCommonTextBlock::SetText(FText InText)
+void UCommonTextBlock::OnTextChanged()
 {
-	Super::SetText(InText);
-	
+	Super::OnTextChanged();
 	if (bAutoCollapseWithEmptyText)
 	{
-		SetVisibility(InText.IsEmpty() ? ESlateVisibility::Collapsed : ESlateVisibility::SelfHitTestInvisible);
+		SetVisibility(GetText().IsEmpty() ? ESlateVisibility::Collapsed : ESlateVisibility::SelfHitTestInvisible);
 	}
+}
 
+void UCommonTextBlock::OnFontChanged()
+{
+	Super::OnFontChanged();
 	if (CommonUIUtils::ShouldDisplayMobileUISizes())
 	{
 		ApplyFontSizeMultiplier();
diff --git a/Engine/Plugins/Experimental/CommonUI/Source/CommonUI/Public/CommonTextBlock.h b/Engine/Plugins/Experimental/CommonUI/Source/CommonUI/Public/CommonTextBlock.h
--- a/Engine/Plugins/Experimental/CommonUI/Source/CommonUI/Public/CommonTextBlock.h
+++ b/Engine/Plugins/Experimental/CommonUI/Source/CommonUI/Public/CommonTextBlock.h
@@ -176,8 +176,6 @@ public:
 	virtual void Serialize(FArchive& Ar) override;
 	virtual void ReleaseSlateResources(bool bReleaseChildren) override;
 
-	virtual void SetText(FText InText) override;
-
 	UFUNCTION(BlueprintCallable, Category = "Common Text")
 	void SetWrapTextWidth(int32 InWrapTextAt);
 
@@ -213,6 +211,9 @@ protected:
 	virtual TSharedRef<SWidget> RebuildWidget() override;
 	virtual void SynchronizeProperties() override;
 
+	virtual void OnTextChanged() override;
+	virtual void OnFontChanged() override;
+
 	/** Mobile font size multiplier. Activated by default on mobile. See CVar Mobile_PreviewFontSize */
 	UPROPERTY(EditAnywhere, BlueprintReadWrite, Getter, Setter, Category = "Mobile", meta = (ClampMin = "0.01", ClampMax = "5.0"))
 	float MobileFontSizeMultiplier = 1.0f;
diff --git a/Engine/Source/Runtime/UMG/Private/Components/TextBlock.cpp b/Engine/Source/Runtime/UMG/Private/Components/TextBlock.cpp
--- a/Engine/Source/Runtime/UMG/Private/Components/TextBlock.cpp
+++ b/Engine/Source/Runtime/UMG/Private/Components/TextBlock.cpp
@@ -126,6 +126,7 @@ void UTextBlock::SetFont(FSlateFontInfo InFontInfo)
 	{
 		MyTextBlock->SetFont(Font);
 	}
+	OnFontChanged();
 }
 
 const FSlateBrush& UTextBlock::GetStrikeBrush() const
@@ -385,6 +386,7 @@ void UTextBlock::SetText(FText InText)
 	{
 		MyTextBlock->SetText(GetDisplayText());
 	}
+	OnTextChanged();
 }
 
 TAttribute<FText> UTextBlock::GetDisplayText()
diff --git a/Engine/Source/Runtime/UMG/Public/Components/TextBlock.h b/Engine/Source/Runtime/UMG/Public/Components/TextBlock.h
--- a/Engine/Source/Runtime/UMG/Public/Components/TextBlock.h
+++ b/Engine/Source/Runtime/UMG/Public/Components/TextBlock.h
@@ -221,6 +221,15 @@ public:
 	UFUNCTION(BlueprintCallable, Category="Appearance")
 	UMaterialInstanceDynamic* GetDynamicOutlineMaterial();
 
+protected:
+	virtual void OnTextChanged()
+	{
+	}
+
+	virtual void OnFontChanged()
+	{
+	}
+
 public:
 	//~ Begin UTextLayoutWidget Interface
 	virtual void SetJustification(ETextJustify::Type InJustification) override;
