Expose CommonUI Icons for EnhancedInput Input Actions

diff --git a/Engine/Plugins/Experimental/CommonUI/CommonUI.uplugin b/Engine/Plugins/Experimental/CommonUI/CommonUI.uplugin
--- a/Engine/Plugins/Experimental/CommonUI/CommonUI.uplugin
+++ b/Engine/Plugins/Experimental/CommonUI/CommonUI.uplugin
@@ -37,6 +37,10 @@
 		{
 			"Name": "GameplayTagsEditor",
 			"Enabled": true
+		},
+		{
+			"Name": "EnhancedInput",
+			"Enabled": true
 		}
 	]
 }
\ No newline at end of file
diff --git a/Engine/Plugins/Experimental/CommonUI/Source/CommonInput/Private/CommonInputSettings.cpp b/Engine/Plugins/Experimental/CommonUI/Source/CommonInput/Private/CommonInputSettings.cpp
--- a/Engine/Plugins/Experimental/CommonUI/Source/CommonInput/Private/CommonInputSettings.cpp
+++ b/Engine/Plugins/Experimental/CommonUI/Source/CommonInput/Private/CommonInputSettings.cpp
@@ -1,6 +1,7 @@
 // Copyright Epic Games, Inc. All Rights Reserved.
 
 #include "CommonInputSettings.h"
+#include "ICommonInputModule.h"
 #include "Engine/DataTable.h"
 #include "Engine/PlatformSettingsManager.h"
 
@@ -133,6 +134,12 @@ FDataTableRowHandle UCommonInputSettings::GetDefaultBackAction() const
 	return FDataTableRowHandle();
 }
 
+bool UCommonInputSettings::IsEnhancedInputSupportEnabled() 
+{
+	static bool bEnabled = ICommonInputModule::Get().GetSettings().GetEnableEnhancedInputSupport();
+	return bEnabled;
+};
+
 void UCommonInputSettings::PostInitProperties()
 {
 	Super::PostInitProperties();
diff --git a/Engine/Plugins/Experimental/CommonUI/Source/CommonInput/Public/CommonInputSettings.h b/Engine/Plugins/Experimental/CommonUI/Source/CommonInput/Public/CommonInputSettings.h
--- a/Engine/Plugins/Experimental/CommonUI/Source/CommonInput/Public/CommonInputSettings.h
+++ b/Engine/Plugins/Experimental/CommonUI/Source/CommonInput/Public/CommonInputSettings.h
@@ -50,6 +50,8 @@ public:
 
 	bool GetEnableDefaultInputConfig() const { return bEnableDefaultInputConfig; }
 
+	bool GetEnableEnhancedInputSupport() const { return bEnableEnhancedInputSupport; }
+
 	TObjectPtr<UCommonInputActionDomainTable> GetActionDomainTable() const { return ActionDomainTablePtr; }
 
 private:
@@ -87,6 +89,10 @@ private:
 	UPROPERTY(config, EditAnywhere, Category = "Input")
 	bool bEnableDefaultInputConfig = true;
 
+	/** Controls if Enhanced Input Support plugin-wide. Requires restart due to caching. */
+	UPROPERTY(config, EditAnywhere, Category = "Input", meta = (ConfigRestartRequired = true))
+	bool bEnableEnhancedInputSupport = false;
+
 	/** Create a derived asset from UCommonInputActionDomainTable to store ordered ActionDomain data for your game */
 	UPROPERTY(config, EditAnywhere, Category = "Action Domain")
 	TSoftObjectPtr<UCommonInputActionDomainTable> ActionDomainTable;
diff --git a/Engine/Plugins/Experimental/CommonUI/Source/CommonUI/CommonUI.Build.cs b/Engine/Plugins/Experimental/CommonUI/Source/CommonUI/CommonUI.Build.cs
--- a/Engine/Plugins/Experimental/CommonUI/Source/CommonUI/CommonUI.Build.cs
+++ b/Engine/Plugins/Experimental/CommonUI/Source/CommonUI/CommonUI.Build.cs
@@ -17,7 +17,8 @@ public class CommonUI : ModuleRules
 				"Slate",
                 "UMG",
                 "WidgetCarousel",
-                "CommonInput",
+				"CommonInput",
+				"EnhancedInput",
 				"MediaAssets",
 				"GameplayTags"
 			}
diff --git a/Engine/Plugins/Experimental/CommonUI/Source/CommonUI/Private/CommonActionWidget.cpp b/Engine/Plugins/Experimental/CommonUI/Source/CommonUI/Private/CommonActionWidget.cpp
--- a/Engine/Plugins/Experimental/CommonUI/Source/CommonUI/Private/CommonActionWidget.cpp
+++ b/Engine/Plugins/Experimental/CommonUI/Source/CommonUI/Private/CommonActionWidget.cpp
@@ -9,6 +9,8 @@
 #include "Widgets/SOverlay.h"
 #include "CommonUITypes.h"
 #include "Materials/MaterialInstanceDynamic.h"
+#include "InputAction.h"
+#include "InputTriggers.h"
 #include "Input/UIActionBinding.h"
 
 #include UE_INLINE_GENERATED_CPP_BY_NAME(CommonActionWidget)
@@ -108,7 +110,9 @@ void UCommonActionWidget::SynchronizeProperties()
 
 FSlateBrush UCommonActionWidget::GetIcon() const
 {
-	return CommonUI::GetIconForInputActions(GetInputSubsystem(), InputActions);
+	return EnhancedInputAction && CommonUI::IsEnhancedInputSupportEnabled()
+		? CommonUI::GetIconForEnhancedInputAction(GetInputSubsystem(), EnhancedInputAction)
+		: CommonUI::GetIconForInputActions(GetInputSubsystem(), InputActions);
 }
 
 UCommonInputSubsystem* UCommonActionWidget::GetInputSubsystem() const
@@ -133,6 +137,11 @@ const FCommonInputActionDataBase* UCommonActionWidget::GetInputActionData() cons
 
 FText UCommonActionWidget::GetDisplayText() const
 {
+	if (EnhancedInputAction && CommonUI::IsEnhancedInputSupportEnabled())
+	{
+		return EnhancedInputAction->ActionDescription;
+	}
+
 	const UCommonInputSubsystem* CommonInputSubsystem = GetInputSubsystem();
 	if (GetGameInstance() && ensure(CommonInputSubsystem))
 	{
@@ -152,6 +161,19 @@ FText UCommonActionWidget::GetDisplayText() const
 
 bool UCommonActionWidget::IsHeldAction() const
 {
+	if (EnhancedInputAction && CommonUI::IsEnhancedInputSupportEnabled())
+	{
+		for (const TObjectPtr<UInputTrigger> Trigger : EnhancedInputAction->Triggers)
+		{
+			if (EnumHasAnyFlags(Trigger->GetSupportedTriggerEvents(), ETriggerEventsSupported::Ongoing))
+			{
+				return true;
+			}
+		}
+
+		return false;
+	}
+
 	const UCommonInputSubsystem* CommonInputSubsystem = GetInputSubsystem();
 	if (GetGameInstance() && ensure(CommonInputSubsystem))
 	{
@@ -165,6 +187,11 @@ bool UCommonActionWidget::IsHeldAction() const
 	return false;
 }
 
+void UCommonActionWidget::SetEnhancedInputAction(UInputAction* InInputAction)
+{
+	EnhancedInputAction = InInputAction;
+}
+
 void UCommonActionWidget::SetInputAction(FDataTableRowHandle InputActionRow)
 {
 	UpdateBindingHandleInternal(FUIActionBindingHandle());
@@ -226,7 +253,8 @@ void UCommonActionWidget::UpdateActionWidget()
 		const UCommonInputSubsystem* CommonInputSubsystem = GetInputSubsystem();
 		if (GetGameInstance() && ensure(CommonInputSubsystem) && CommonInputSubsystem->ShouldShowInputKeys())
 		{
-			if (const FCommonInputActionDataBase* InputActionData = GetInputActionData())
+			const FCommonInputActionDataBase* InputActionData = GetInputActionData();
+			if (InputActionData || (EnhancedInputAction && CommonUI::IsEnhancedInputSupportEnabled()))
 			{
 				if (bAlwaysHideOverride)
 				{
@@ -251,7 +279,7 @@ void UCommonActionWidget::UpdateActionWidget()
 							MyIcon->Invalidate(EInvalidateWidgetReason::Layout);
 						}
 
-						if (InputActionData->GetCurrentInputTypeInfo(CommonInputSubsystem).bActionRequiresHold)
+						if (IsHeldAction())
 						{
 							MyProgressImage->SetVisibility(EVisibility::SelfHitTestInvisible);
 						}
diff --git a/Engine/Plugins/Experimental/CommonUI/Source/CommonUI/Private/CommonUITypes.cpp b/Engine/Plugins/Experimental/CommonUI/Source/CommonUI/Private/CommonUITypes.cpp
--- a/Engine/Plugins/Experimental/CommonUI/Source/CommonUI/Private/CommonUITypes.cpp
+++ b/Engine/Plugins/Experimental/CommonUI/Source/CommonUI/Private/CommonUITypes.cpp
@@ -1,10 +1,15 @@
 // Copyright Epic Games, Inc. All Rights Reserved.
  
 #include "CommonUITypes.h"
+
 #include "CommonInputBaseTypes.h"
-#include "CommonUIPrivate.h"
 #include "CommonInputSubsystem.h"
+#include "CommonUIPrivate.h"
+#include "EnhancedInputSubsystems.h"
+#include "Engine/LocalPlayer.h"
 #include "HAL/PlatformInput.h"
+#include "ICommonInputModule.h"
+#include "InputAction.h"
 #include "Styling/SlateTypes.h"
 #include "Styling/StyleDefaults.h"
 
@@ -223,3 +228,29 @@ FSlateBrush CommonUI::GetIconForInputActions(const UCommonInputSubsystem* Common
 
 	return *FStyleDefaults::GetNoBrush();
 }
+
+bool CommonUI::IsEnhancedInputSupportEnabled()
+{
+	static bool bEnabled = ICommonInputModule::Get().GetSettings().GetEnableEnhancedInputSupport();
+	return bEnabled;
+}
+
+FSlateBrush CommonUI::GetIconForEnhancedInputAction(const UCommonInputSubsystem* CommonInputSubsystem, const TObjectPtr<UInputAction> InputAction)
+{
+	TArray<FKey> Keys;
+	if (ULocalPlayer* LocalPlayer = CommonInputSubsystem->GetLocalPlayer())
+	{
+		if (UEnhancedInputLocalPlayerSubsystem* EnhancedInputLocalPlayerSubsystem = LocalPlayer->GetSubsystem<UEnhancedInputLocalPlayerSubsystem>())
+		{
+			Keys = EnhancedInputLocalPlayerSubsystem->QueryKeysMappedToAction(InputAction);
+		}
+	}
+
+	FSlateBrush SlateBrush;
+	if (Keys.Num() && UCommonInputPlatformSettings::Get()->TryGetInputBrush(SlateBrush, Keys[0], CommonInputSubsystem->GetCurrentInputType(), CommonInputSubsystem->GetCurrentGamepadName()))
+	{
+		return SlateBrush;
+	}
+
+	return *FStyleDefaults::GetNoBrush();
+}
\ No newline at end of file
diff --git a/Engine/Plugins/Experimental/CommonUI/Source/CommonUI/Public/CommonActionWidget.h b/Engine/Plugins/Experimental/CommonUI/Source/CommonUI/Public/CommonActionWidget.h
--- a/Engine/Plugins/Experimental/CommonUI/Source/CommonUI/Public/CommonActionWidget.h
+++ b/Engine/Plugins/Experimental/CommonUI/Source/CommonUI/Public/CommonActionWidget.h
@@ -42,6 +42,9 @@ public:
 	UFUNCTION(BlueprintCallable, Category = CommonActionWidget)
 	FText GetDisplayText() const;
 
+	UFUNCTION(BlueprintCallable, Category = CommonActionWidget)
+	void SetEnhancedInputAction(UInputAction* InInputAction);
+
 	UFUNCTION(BlueprintCallable, Category = CommonActionWidget)
 	void SetInputAction(FDataTableRowHandle InputActionRow);
 
@@ -92,6 +95,12 @@ protected:
 	UPROPERTY(EditAnywhere, BlueprintReadOnly, Category = CommonActionWidget, meta = (RowType = "/Script/CommonUI.CommonInputActionDataBase", TitleProperty = "RowName"))
 	TArray<FDataTableRowHandle> InputActions;
 
+	/**
+	 * Input Action this common action widget is intended to represent. Optional if using EnhancedInputs
+	 */
+	UPROPERTY(EditAnywhere, BlueprintReadOnly, Category = CommonActionWidget)
+	TObjectPtr<class UInputAction> EnhancedInputAction;
+
 	//@todo DanH: Create clearer split between support for the new & legacy system in here
 	FUIActionBindingHandle DisplayedBindingHandle;
 
diff --git a/Engine/Plugins/Experimental/CommonUI/Source/CommonUI/Public/CommonUITypes.h b/Engine/Plugins/Experimental/CommonUI/Source/CommonUI/Public/CommonUITypes.h
--- a/Engine/Plugins/Experimental/CommonUI/Source/CommonUI/Public/CommonUITypes.h
+++ b/Engine/Plugins/Experimental/CommonUI/Source/CommonUI/Public/CommonUITypes.h
@@ -173,6 +173,9 @@ public:
 	static FScrollBoxStyle EmptyScrollBoxStyle;
 	static const FCommonInputActionDataBase* GetInputActionData(const FDataTableRowHandle& InputActionRowHandle);
 	static FSlateBrush GetIconForInputActions(const UCommonInputSubsystem* CommonInputSubsystem, const TArray<FDataTableRowHandle>& InputActions);
+
+	static bool IsEnhancedInputSupportEnabled();
+	static FSlateBrush GetIconForEnhancedInputAction(const UCommonInputSubsystem* CommonInputSubsystem, const TObjectPtr<class UInputAction> InputAction);
 };
 
 DECLARE_DYNAMIC_DELEGATE_OneParam(FOnItemClicked, UUserWidget*, Widget);
