Fix compile issue additional function defintion issue due to bad unshelve.

diff --git a/Engine/Plugins/Experimental/CommonUI/Source/CommonInput/Private/CommonInputSettings.cpp b/Engine/Plugins/Experimental/CommonUI/Source/CommonInput/Private/CommonInputSettings.cpp
--- a/Engine/Plugins/Experimental/CommonUI/Source/CommonInput/Private/CommonInputSettings.cpp
+++ b/Engine/Plugins/Experimental/CommonUI/Source/CommonInput/Private/CommonInputSettings.cpp
@@ -1,7 +1,6 @@
 // Copyright Epic Games, Inc. All Rights Reserved.
 
 #include "CommonInputSettings.h"
-#include "ICommonInputModule.h"
 #include "Engine/DataTable.h"
 #include "Engine/PlatformSettingsManager.h"
 
@@ -134,12 +133,6 @@ FDataTableRowHandle UCommonInputSettings::GetDefaultBackAction() const
 	return FDataTableRowHandle();
 }
 
-bool UCommonInputSettings::IsEnhancedInputSupportEnabled() 
-{
-	static bool bEnabled = ICommonInputModule::Get().GetSettings().GetEnableEnhancedInputSupport();
-	return bEnabled;
-};
-
 void UCommonInputSettings::PostInitProperties()
 {
 	Super::PostInitProperties();
