CommonButtonBase - Remove hovering for touch input through CVar in an attempt to fix several issues related to hovering on touch input devices.

diff --git a/Engine/Plugins/Experimental/CommonUI/Source/CommonUI/Private/CommonButtonBase.cpp b/Engine/Plugins/Experimental/CommonUI/Source/CommonUI/Private/CommonButtonBase.cpp
--- a/Engine/Plugins/Experimental/CommonUI/Source/CommonUI/Private/CommonButtonBase.cpp
+++ b/Engine/Plugins/Experimental/CommonUI/Source/CommonUI/Private/CommonButtonBase.cpp
@@ -23,6 +23,16 @@
 
 #include UE_INLINE_GENERATED_CPP_BY_NAME(CommonButtonBase)
 
+namespace UE::CommonUI::Private
+{
+	int32 bSimulateHoverOnTouchInput = 1;
+	FAutoConsoleVariableRef CVarSimulateHoverOnTouchInput(
+		TEXT("CommonButton.SimulateHoverForTouch"),
+		bSimulateHoverOnTouchInput,
+		TEXT("Enables or disables the hovering/unhovering simulation on Touch Input System.\n0: Disable, 1: Enable (default)"),
+		ECVF_Default);
+}
+
 //////////////////////////////////////////////////////////////////////////
 // UCommonButtonStyle
 //////////////////////////////////////////////////////////////////////////
@@ -1094,7 +1104,7 @@ void UCommonButtonBase::HandleButtonPressed()
 
 	UCommonInputSubsystem* CommonInputSubsystem = GetInputSubsystem();
 
-	if (CommonInputSubsystem && CommonInputSubsystem->GetCurrentInputType() == ECommonInputType::Touch)
+	if (CommonInputSubsystem && CommonInputSubsystem->GetCurrentInputType() == ECommonInputType::Touch && UE::CommonUI::Private::bSimulateHoverOnTouchInput)
 	{
 		// Simulate hover events when using touch input
 		NativeOnHovered();
@@ -1107,7 +1117,7 @@ void UCommonButtonBase::HandleButtonReleased()
 
 	UCommonInputSubsystem* CommonInputSubsystem = GetInputSubsystem();
 
-	if (CommonInputSubsystem && CommonInputSubsystem->GetCurrentInputType() == ECommonInputType::Touch)
+	if (CommonInputSubsystem && CommonInputSubsystem->GetCurrentInputType() == ECommonInputType::Touch && UE::CommonUI::Private::bSimulateHoverOnTouchInput)
 	{
 		// Simulate hover events when using touch input
 		NativeOnUnhovered();
