Speculative fix for failed 'check' of a TWeakObjectPtr.

Replaced with an ensureMsgf to see if we get better logging in future reports.

#tests Compiled DevelopmentEditor Win64 / Ran a PIE instance with a breakpoint in relevant code and verified by default we did not enter the failed ensure/return case.

diff --git a/Engine/Plugins/Experimental/CommonUI/Source/CommonUI/Private/Input/UIActionRouterTypes.cpp b/Engine/Plugins/Experimental/CommonUI/Source/CommonUI/Private/Input/UIActionRouterTypes.cpp
--- a/Engine/Plugins/Experimental/CommonUI/Source/CommonUI/Private/Input/UIActionRouterTypes.cpp
+++ b/Engine/Plugins/Experimental/CommonUI/Source/CommonUI/Private/Input/UIActionRouterTypes.cpp
@@ -1263,7 +1263,11 @@ void FActivatableTreeRoot::FocusLeafmostNode()
 
 	FActivatableTreeNodePtr PinnedLeafmostNode = LeafmostActiveNode.Pin();
 	UCommonActivatableWidget* LeafWidget = PinnedLeafmostNode->GetWidget();
-	check(LeafWidget);
+	
+	if (!ensureMsgf(LeafWidget != nullptr, TEXT("Unable to focus leafmost node because the TWeakObjectPtr LeafWidget is nullptr")))
+	{
+		return;
+	}
 
 	const int32 OwnerSlateId = GetOwnerUserIndex();
 	ULocalPlayer& LocalPlayer = *GetActionRouter().GetLocalPlayerChecked();
