Fix issue where multiple default action prompts appear on buttons when disabled and re-enabled.
Problem: When a focused button is disabled, focus is changed, then button is re-enabled - it will continue to show the default action prompt.
Solution: When enabling/disabling a button - update the action prompt and remove the default action when disabled.

diff --git a/Engine/Plugins/Experimental/CommonUI/Source/CommonUI/Private/CommonButtonBase.cpp b/Engine/Plugins/Experimental/CommonUI/Source/CommonUI/Private/CommonButtonBase.cpp
--- a/Engine/Plugins/Experimental/CommonUI/Source/CommonUI/Private/CommonButtonBase.cpp
+++ b/Engine/Plugins/Experimental/CommonUI/Source/CommonUI/Private/CommonButtonBase.cpp
@@ -1034,8 +1034,8 @@ void UCommonButtonBase::UpdateInputActionWidget()
 		{
 			InputActionWidget->SetInputAction(TriggeredInputAction);
 		}
-		// Visualize the default click action when neither input action is bound and when the widget is hovered
-		else if (bShouldUseFallbackDefaultInputAction)
+		// Visualize the default click action when neither input action is bound and when the widget is enabled and hovered
+		else if (bShouldUseFallbackDefaultInputAction && bButtonEnabled)
 		{
 			FDataTableRowHandle HoverStateHandle;
 			if (IsHovered())
@@ -1525,6 +1525,7 @@ void UCommonButtonBase::EnableButton()
 
 		if (InputActionWidget)
 		{
+			UpdateInputActionWidget();
 			InputActionWidget->SetIsEnabled(bButtonEnabled);
 		}
 	}
@@ -1543,6 +1544,7 @@ void UCommonButtonBase::DisableButton()
 
 		if (InputActionWidget)
 		{
+			UpdateInputActionWidget();
 			InputActionWidget->SetIsEnabled(bButtonEnabled);
 		}
 	}
