[24.00] UCommonTabListWidgetBase:
- Made GetTabButtonBaseByID const.

diff --git a/Engine/Plugins/Experimental/CommonUI/Source/CommonUI/Private/CommonTabListWidgetBase.cpp b/Engine/Plugins/Experimental/CommonUI/Source/CommonUI/Private/CommonTabListWidgetBase.cpp
--- a/Engine/Plugins/Experimental/CommonUI/Source/CommonUI/Private/CommonTabListWidgetBase.cpp
+++ b/Engine/Plugins/Experimental/CommonUI/Source/CommonUI/Private/CommonTabListWidgetBase.cpp
@@ -328,9 +328,9 @@ void UCommonTabListWidgetBase::DisableTabWithReason(FName TabNameID, const FText
 	}
 }
 
-UCommonButtonBase* UCommonTabListWidgetBase::GetTabButtonBaseByID(FName TabNameID)
+UCommonButtonBase* UCommonTabListWidgetBase::GetTabButtonBaseByID(FName TabNameID) const
 {
-	if (FCommonRegisteredTabInfo* TabInfo = RegisteredTabsByID.Find(TabNameID))
+	if (const FCommonRegisteredTabInfo* TabInfo = RegisteredTabsByID.Find(TabNameID))
 	{
 		return TabInfo->TabButton;
 	}
diff --git a/Engine/Plugins/Experimental/CommonUI/Source/CommonUI/Public/CommonTabListWidgetBase.h b/Engine/Plugins/Experimental/CommonUI/Source/CommonUI/Public/CommonTabListWidgetBase.h
--- a/Engine/Plugins/Experimental/CommonUI/Source/CommonUI/Public/CommonTabListWidgetBase.h
+++ b/Engine/Plugins/Experimental/CommonUI/Source/CommonUI/Public/CommonTabListWidgetBase.h
@@ -142,7 +142,7 @@ public:
 
 	/** Returns the tab button matching the ID, if found */
 	UFUNCTION(BlueprintCallable, Category = TabList)
-	UCommonButtonBase* GetTabButtonBaseByID(FName TabNameID);
+	UCommonButtonBase* GetTabButtonBaseByID(FName TabNameID) const;
 
 protected:
 	// UUserWidget interface
