Hide EnhancedInput in CommonActionWidget via EditConditions

diff --git a/Engine/Plugins/Experimental/CommonUI/Source/CommonInput/Private/CommonInputSettings.cpp b/Engine/Plugins/Experimental/CommonUI/Source/CommonInput/Private/CommonInputSettings.cpp
--- a/Engine/Plugins/Experimental/CommonUI/Source/CommonInput/Private/CommonInputSettings.cpp
+++ b/Engine/Plugins/Experimental/CommonUI/Source/CommonInput/Private/CommonInputSettings.cpp
@@ -1,6 +1,7 @@
 // Copyright Epic Games, Inc. All Rights Reserved.
 
 #include "CommonInputSettings.h"
+#include "ICommonInputModule.h"
 #include "Engine/DataTable.h"
 #include "Engine/PlatformSettingsManager.h"
 
@@ -133,6 +134,12 @@ FDataTableRowHandle UCommonInputSettings::GetDefaultBackAction() const
 	return FDataTableRowHandle();
 }
 
+bool UCommonInputSettings::IsEnhancedInputSupportEnabled() 
+{
+	static bool bEnabled = ICommonInputModule::Get().GetSettings().GetEnableEnhancedInputSupport();
+	return bEnabled;
+};
+
 void UCommonInputSettings::PostInitProperties()
 {
 	Super::PostInitProperties();
diff --git a/Engine/Plugins/Experimental/CommonUI/Source/CommonInput/Public/CommonInputSettings.h b/Engine/Plugins/Experimental/CommonUI/Source/CommonInput/Public/CommonInputSettings.h
--- a/Engine/Plugins/Experimental/CommonUI/Source/CommonInput/Public/CommonInputSettings.h
+++ b/Engine/Plugins/Experimental/CommonUI/Source/CommonInput/Public/CommonInputSettings.h
@@ -54,6 +54,12 @@ public:
 
 	TObjectPtr<UCommonInputActionDomainTable> GetActionDomainTable() const { return ActionDomainTablePtr; }
 
+public:
+
+	/** Static version of enhanced input support check, exists to hide based on edit condition */
+	UFUNCTION()
+	static bool IsEnhancedInputSupportEnabled(); 
+
 private:
 	virtual void PostInitProperties() override;
 
diff --git a/Engine/Plugins/Experimental/CommonUI/Source/CommonUI/Public/CommonActionWidget.h b/Engine/Plugins/Experimental/CommonUI/Source/CommonUI/Public/CommonActionWidget.h
--- a/Engine/Plugins/Experimental/CommonUI/Source/CommonUI/Public/CommonActionWidget.h
+++ b/Engine/Plugins/Experimental/CommonUI/Source/CommonUI/Public/CommonActionWidget.h
@@ -98,7 +98,7 @@ protected:
 	/**
 	 * Input Action this common action widget is intended to represent. Optional if using EnhancedInputs
 	 */
-	UPROPERTY(EditAnywhere, BlueprintReadOnly, Category = CommonActionWidget)
+	UPROPERTY(EditAnywhere, BlueprintReadOnly, Category = CommonActionWidget, meta = (EditCondition = "CommonInput.CommonInputSettings.IsEnhancedInputSupportEnabled", EditConditionHides))
 	TObjectPtr<class UInputAction> EnhancedInputAction;
 
 	//@todo DanH: Create clearer split between support for the new & legacy system in here
