Expose virtual native function on Bound Action Bar for derived classes to react when new button is created.
#tests: Verified functionality by creating derived class that performed custom actions on newly created buttons.

diff --git a/Engine/Plugins/Experimental/CommonUI/Source/CommonUI/Private/Input/CommonBoundActionBar.cpp b/Engine/Plugins/Experimental/CommonUI/Source/CommonUI/Private/Input/CommonBoundActionBar.cpp
--- a/Engine/Plugins/Experimental/CommonUI/Source/CommonUI/Private/Input/CommonBoundActionBar.cpp
+++ b/Engine/Plugins/Experimental/CommonUI/Source/CommonUI/Private/Input/CommonBoundActionBar.cpp
@@ -247,6 +247,7 @@ void UCommonBoundActionBar::HandleDeferredDisplayUpdate()
 						if (ensure(ActionButton))
 						{
 							ActionButton->SetRepresentedAction(BindingHandle);
+							NativeOnActionButtonCreated(ActionButton, BindingHandle);
 						}
 					}
 				}
diff --git a/Engine/Plugins/Experimental/CommonUI/Source/CommonUI/Public/Input/CommonBoundActionBar.h b/Engine/Plugins/Experimental/CommonUI/Source/CommonUI/Public/Input/CommonBoundActionBar.h
--- a/Engine/Plugins/Experimental/CommonUI/Source/CommonUI/Public/Input/CommonBoundActionBar.h
+++ b/Engine/Plugins/Experimental/CommonUI/Source/CommonUI/Public/Input/CommonBoundActionBar.h
@@ -7,8 +7,11 @@
 #include "Tickable.h"
 #include "CommonBoundActionBar.generated.h"
 
-class IWidgetCompilerLog;
+class ICommonBoundActionButtonInterface;
 class IConsoleVariable;
+class IWidgetCompilerLog;
+
+struct FUIActionBindingHandle;
 
 /**
  * A box populated with current actions available per CommonUI's Input Handler.
@@ -34,6 +37,8 @@ protected:
 	virtual void SynchronizeProperties() override;
 	virtual void ReleaseSlateResources(bool bReleaseChildren) override;
 
+	virtual void NativeOnActionButtonCreated(ICommonBoundActionButtonInterface* ActionButton, FUIActionBindingHandle const& RepresentedAction) { }
+
 #if WITH_EDITOR
 	void ValidateCompiledDefaults(IWidgetCompilerLog& CompileLog) const override;
 #endif
