CommonUI - Removing IsVirtualAcceptPressedBound from action router base, it's not actually implemented.

diff --git a/Engine/Plugins/Experimental/CommonUI/Source/CommonUI/Public/Input/CommonUIActionRouterBase.h b/Engine/Plugins/Experimental/CommonUI/Source/CommonUI/Public/Input/CommonUIActionRouterBase.h
--- a/Engine/Plugins/Experimental/CommonUI/Source/CommonUI/Public/Input/CommonUIActionRouterBase.h
+++ b/Engine/Plugins/Experimental/CommonUI/Source/CommonUI/Public/Input/CommonUIActionRouterBase.h
@@ -116,8 +116,6 @@ public:
 	void AddBinding(FUIActionBindingHandle Binding);
 	void RemoveBinding(FUIActionBindingHandle Binding);
 
-	bool IsVirtualAcceptPressedBound() const;
-
 	int32 GetLocalPlayerIndex() const;
 
 	void RefreshActiveRootFocus();
