Add support for Enhanced Input Actions as UI triggers in CommonUI.
Add Enhanced Input Action bindings as an option where data tables existed previously.
Add application & removal of IMC's based on widget activation & deactivation.

diff --git a/Engine/Plugins/Experimental/CommonUI/Source/CommonInput/CommonInput.Build.cs b/Engine/Plugins/Experimental/CommonUI/Source/CommonInput/CommonInput.Build.cs
--- a/Engine/Plugins/Experimental/CommonUI/Source/CommonInput/CommonInput.Build.cs
+++ b/Engine/Plugins/Experimental/CommonUI/Source/CommonInput/CommonInput.Build.cs
@@ -13,6 +13,7 @@ public class CommonInput : ModuleRules
                 "CoreUObject",
                 "Engine",
                 "InputCore",
+				"EnhancedInput",
 				"DeveloperSettings"
             }
 		);
diff --git a/Engine/Plugins/Experimental/CommonUI/Source/CommonInput/Private/CommonInputSettings.cpp b/Engine/Plugins/Experimental/CommonUI/Source/CommonInput/Private/CommonInputSettings.cpp
--- a/Engine/Plugins/Experimental/CommonUI/Source/CommonInput/Private/CommonInputSettings.cpp
+++ b/Engine/Plugins/Experimental/CommonUI/Source/CommonInput/Private/CommonInputSettings.cpp
@@ -1,6 +1,7 @@
 // Copyright Epic Games, Inc. All Rights Reserved.
 
 #include "CommonInputSettings.h"
+
 #include "ICommonInputModule.h"
 #include "Engine/DataTable.h"
 #include "Engine/PlatformSettingsManager.h"
@@ -134,6 +135,28 @@ FDataTableRowHandle UCommonInputSettings::GetDefaultBackAction() const
 	return FDataTableRowHandle();
 }
 
+UInputAction* UCommonInputSettings::GetEnhancedInputClickAction() const
+{
+	ensureMsgf(bInputDataLoaded, TEXT("Warning, CommonUI input data not loaded."));
+
+	if (InputDataClass && IsEnhancedInputSupportEnabled())
+	{
+		return InputDataClass.GetDefaultObject()->EnhancedInputClickAction;
+	}
+	return nullptr;
+}
+
+UInputAction* UCommonInputSettings::GetEnhancedInputBackAction() const
+{
+	ensureMsgf(bInputDataLoaded, TEXT("Warning, CommonUI input data not loaded."));
+
+	if (InputDataClass && IsEnhancedInputSupportEnabled())
+	{
+		return InputDataClass.GetDefaultObject()->EnhancedInputBackAction;
+	}
+	return nullptr;
+}
+
 bool UCommonInputSettings::IsEnhancedInputSupportEnabled() 
 {
 	static bool bEnabled = ICommonInputModule::Get().GetSettings().GetEnableEnhancedInputSupport();
diff --git a/Engine/Plugins/Experimental/CommonUI/Source/CommonInput/Public/CommonInputBaseTypes.h b/Engine/Plugins/Experimental/CommonUI/Source/CommonInput/Public/CommonInputBaseTypes.h
--- a/Engine/Plugins/Experimental/CommonUI/Source/CommonInput/Public/CommonInputBaseTypes.h
+++ b/Engine/Plugins/Experimental/CommonUI/Source/CommonInput/Public/CommonInputBaseTypes.h
@@ -14,12 +14,14 @@
 #include "Engine/DataTable.h"
 #include "UObject/SoftObjectPtr.h"
 #include "Engine/PlatformSettings.h"
+#include "InputAction.h"
+
 #include "CommonInputBaseTypes.generated.h"
 
 
 class UTexture2D;
 class UMaterial;
-class UCommonInputSettings;
+class UCommonInputSettings; 
 
 struct COMMONINPUT_API FCommonInputDefaults
 {
@@ -90,6 +92,12 @@ public:
 
 	UPROPERTY(EditDefaultsOnly, Category = "Properties", meta = (RowType = "/Script/CommonUI.CommonInputActionDataBase"))
 	FDataTableRowHandle DefaultBackAction;
+
+	UPROPERTY(EditDefaultsOnly, Category = "Properties", meta = (EditCondition = "CommonInput.CommonInputSettings.IsEnhancedInputSupportEnabled", EditConditionHides))
+	TObjectPtr<UInputAction> EnhancedInputClickAction;
+
+	UPROPERTY(EditDefaultsOnly, Category = "Properties", meta = (EditCondition = "CommonInput.CommonInputSettings.IsEnhancedInputSupportEnabled", EditConditionHides))
+	TObjectPtr<UInputAction> EnhancedInputBackAction;
 };
 
 /* Derive from this class to store the Input data. It is referenced in the Common Input Settings, found in the project settings UI. */
diff --git a/Engine/Plugins/Experimental/CommonUI/Source/CommonInput/Public/CommonInputSettings.h b/Engine/Plugins/Experimental/CommonUI/Source/CommonInput/Public/CommonInputSettings.h
--- a/Engine/Plugins/Experimental/CommonUI/Source/CommonInput/Public/CommonInputSettings.h
+++ b/Engine/Plugins/Experimental/CommonUI/Source/CommonInput/Public/CommonInputSettings.h
@@ -17,6 +17,8 @@
 #include "Engine/PlatformSettings.h"
 #include "CommonInputSettings.generated.h"
 
+class UInputAction;
+
 UCLASS(config = Game, defaultconfig)
 class COMMONINPUT_API UCommonInputSettings : public UDeveloperSettings
 {
@@ -38,6 +40,9 @@ public:
 	FDataTableRowHandle GetDefaultClickAction() const;
 	FDataTableRowHandle GetDefaultBackAction() const;
 
+	UInputAction* GetEnhancedInputClickAction() const;
+	UInputAction* GetEnhancedInputBackAction() const;
+
 	bool GetEnableInputMethodThrashingProtection() const { return bEnableInputMethodThrashingProtection; }
 
 	int32 GetInputMethodThrashingLimit() const { return InputMethodThrashingLimit; }
diff --git a/Engine/Plugins/Experimental/CommonUI/Source/CommonUI/Private/CommonActionWidget.cpp b/Engine/Plugins/Experimental/CommonUI/Source/CommonUI/Private/CommonActionWidget.cpp
--- a/Engine/Plugins/Experimental/CommonUI/Source/CommonUI/Private/CommonActionWidget.cpp
+++ b/Engine/Plugins/Experimental/CommonUI/Source/CommonUI/Private/CommonActionWidget.cpp
@@ -189,7 +189,9 @@ bool UCommonActionWidget::IsHeldAction() const
 
 void UCommonActionWidget::SetEnhancedInputAction(UInputAction* InInputAction)
 {
+	UpdateBindingHandleInternal(FUIActionBindingHandle());
 	EnhancedInputAction = InInputAction;
+	UpdateActionWidget();
 }
 
 void UCommonActionWidget::SetInputAction(FDataTableRowHandle InputActionRow)
diff --git a/Engine/Plugins/Experimental/CommonUI/Source/CommonUI/Private/CommonActivatableWidget.cpp b/Engine/Plugins/Experimental/CommonUI/Source/CommonUI/Private/CommonActivatableWidget.cpp
--- a/Engine/Plugins/Experimental/CommonUI/Source/CommonUI/Private/CommonActivatableWidget.cpp
+++ b/Engine/Plugins/Experimental/CommonUI/Source/CommonUI/Private/CommonActivatableWidget.cpp
@@ -3,7 +3,10 @@
 #include "CommonActivatableWidget.h"
 #include "CommonInputSettings.h"
 #include "CommonUIPrivate.h"
+#include "CommonUITypes.h"
 #include "Engine/GameInstance.h"
+#include "EnhancedInputSubsystems.h"
+#include "InputMappingContext.h"
 #include "Input/CommonUIInputTypes.h"
 #include "Input/UIActionBinding.h"
 #include "Input/UIActionRouterTypes.h"
@@ -20,10 +23,20 @@ void UCommonActivatableWidget::NativeConstruct()
 
 	if (bIsBackHandler)
 	{
-		FBindUIActionArgs BindArgs(ICommonInputModule::GetSettings().GetDefaultBackAction(), FSimpleDelegate::CreateUObject(this, &UCommonActivatableWidget::HandleBackAction));
-		BindArgs.bDisplayInActionBar = bIsBackActionDisplayedInActionBar;
+		if (CommonUI::IsEnhancedInputSupportEnabled() && ICommonInputModule::GetSettings().GetEnhancedInputBackAction())
+		{
+			FBindUIActionArgs BindArgs(ICommonInputModule::GetSettings().GetEnhancedInputBackAction(), FSimpleDelegate::CreateUObject(this, &UCommonActivatableWidget::HandleBackAction));
+			BindArgs.bDisplayInActionBar = bIsBackActionDisplayedInActionBar;
 
-		DefaultBackActionHandle = RegisterUIActionBinding(BindArgs);
+			DefaultBackActionHandle = RegisterUIActionBinding(BindArgs);
+		}
+		else
+		{
+			FBindUIActionArgs BindArgs(ICommonInputModule::GetSettings().GetDefaultBackAction(), FSimpleDelegate::CreateUObject(this, &UCommonActivatableWidget::HandleBackAction));
+			BindArgs.bDisplayInActionBar = bIsBackActionDisplayedInActionBar;
+
+			DefaultBackActionHandle = RegisterUIActionBinding(BindArgs);
+		}
 	}
 
 	if (bAutoActivate)
@@ -220,6 +233,17 @@ void UCommonActivatableWidget::NativeOnActivated()
 			UE_LOG(LogCommonUI, Verbose, TEXT("[%s] set visibility to [%s] on activation"), *GetName(), *StaticEnum<ESlateVisibility>()->GetDisplayValueAsText(ActivatedVisibility).ToString());
 		}
 
+		if (CommonUI::IsEnhancedInputSupportEnabled() && InputMapping)
+		{
+			if (const ULocalPlayer* LocalPlayer = GetOwningLocalPlayer())
+			{
+				if (UEnhancedInputLocalPlayerSubsystem* InputSystem = LocalPlayer->GetSubsystem<UEnhancedInputLocalPlayerSubsystem>())
+				{
+					InputSystem->AddMappingContext(InputMapping, InputMappingPriority);
+				}
+			}
+		}
+
 		BP_OnActivated();
 		OnActivated().Broadcast();
 		BP_OnWidgetActivated.Broadcast();
@@ -236,6 +260,17 @@ void UCommonActivatableWidget::NativeOnDeactivated()
 			UE_LOG(LogCommonUI, Verbose, TEXT("[%s] set visibility to [%d] on deactivation"), *GetName(), *StaticEnum<ESlateVisibility>()->GetDisplayValueAsText(DeactivatedVisibility).ToString());
 		}
 
+		if (CommonUI::IsEnhancedInputSupportEnabled() && InputMapping)
+		{
+			if (const ULocalPlayer* LocalPlayer = GetOwningLocalPlayer())
+			{
+				if (UEnhancedInputLocalPlayerSubsystem* InputSystem = LocalPlayer->GetSubsystem<UEnhancedInputLocalPlayerSubsystem>())
+				{
+					InputSystem->RemoveMappingContext(InputMapping);
+				}
+			}
+		}
+
 		// Cancel any holds that were active
 		ClearActiveHoldInputs();
 
diff --git a/Engine/Plugins/Experimental/CommonUI/Source/CommonUI/Private/CommonButtonBase.cpp b/Engine/Plugins/Experimental/CommonUI/Source/CommonUI/Private/CommonButtonBase.cpp
--- a/Engine/Plugins/Experimental/CommonUI/Source/CommonUI/Private/CommonButtonBase.cpp
+++ b/Engine/Plugins/Experimental/CommonUI/Source/CommonUI/Private/CommonButtonBase.cpp
@@ -10,6 +10,7 @@
 #include "Components/ButtonSlot.h"
 #include "Blueprint/WidgetTree.h"
 #include "CommonButtonTypes.h"
+#include "CommonUITypes.h"
 #include "ICommonUIModule.h"
 #include "Widgets/Layout/SBox.h"
 #include "Framework/Application/SlateApplication.h"
@@ -18,6 +19,7 @@
 #include "ICommonInputModule.h"
 #include "CommonInputSettings.h"
 #include "Input/CommonUIInputTypes.h"
+#include "InputAction.h"
 #include "Sound/SoundBase.h"
 #include "Styling/UMGCoreStyle.h"
 
@@ -502,6 +504,19 @@ void UCommonButtonBase::OnInputMethodChanged(ECommonInputType CurrentInputType)
 
 void UCommonButtonBase::BindTriggeringInputActionToClick()
 {
+	if (CommonUI::IsEnhancedInputSupportEnabled() && TriggeringEnhancedInputAction)
+	{
+		FBindUIActionArgs BindArgs(TriggeringEnhancedInputAction, false, FSimpleDelegate::CreateUObject(this, &UCommonButtonBase::HandleTriggeringActionCommited));
+		BindArgs.OnHoldActionProgressed.BindUObject(this, &UCommonButtonBase::NativeOnActionProgress);
+		BindArgs.bIsPersistent = bIsPersistentBinding;
+
+		BindArgs.InputMode = InputModeOverride;
+
+		TriggeringBindingHandle = RegisterUIActionBinding(BindArgs);
+
+		return;
+	}
+
 	if (TriggeringInputAction.IsNull() || !TriggeredInputAction.IsNull())
 	{
 		return;
@@ -521,6 +536,13 @@ void UCommonButtonBase::BindTriggeringInputActionToClick()
 
 void UCommonButtonBase::UnbindTriggeringInputActionToClick()
 {
+	if (CommonUI::IsEnhancedInputSupportEnabled() && TriggeringEnhancedInputAction)
+	{
+		TriggeringBindingHandle.Unregister();
+
+		return;
+	}
+
 	if (TriggeringInputAction.IsNull() || !TriggeredInputAction.IsNull())
 	{
 		return;
@@ -988,6 +1010,26 @@ void UCommonButtonBase::SetTriggeringInputAction(const FDataTableRowHandle & Inp
 	}
 }
 
+void UCommonButtonBase::SetTriggeringEnhancedInputAction(UInputAction* InInputAction)
+{
+	if (CommonUI::IsEnhancedInputSupportEnabled() && TriggeringEnhancedInputAction != InInputAction)
+	{
+		UnbindTriggeringInputActionToClick();
+
+		TriggeringEnhancedInputAction = InInputAction;
+
+		if (!IsDesignTime())
+		{
+			BindTriggeringInputActionToClick();
+		}
+
+		// Update the Input action widget whenever the triggering input action changes
+		UpdateInputActionWidget();
+
+		OnTriggeringEnhancedInputActionChanged(InInputAction);
+	}
+}
+
 bool UCommonButtonBase::GetInputAction(FDataTableRowHandle &InputActionRow) const
 {
 	bool bBothActionsSet = !TriggeringInputAction.IsNull() && !TriggeredInputAction.IsNull();
@@ -1010,6 +1052,11 @@ bool UCommonButtonBase::GetInputAction(FDataTableRowHandle &InputActionRow) cons
 	}
 }
 
+UInputAction* UCommonButtonBase::GetEnhancedInputAction() const
+{
+	return TriggeringEnhancedInputAction;
+}
+
 UMaterialInstanceDynamic* UCommonButtonBase::GetSingleMaterialStyleMID() const
 {
 	return SingleMaterialStyleMID;
@@ -1024,8 +1071,13 @@ void UCommonButtonBase::UpdateInputActionWidget()
 	// Update the input action state of the input action widget contextually based on the current state of the button
 	if (GetGameInstance() && InputActionWidget)
 	{
-		// Prefer visualizing the triggering input action before all else
-		if (!TriggeringInputAction.IsNull())
+		// Prefer visualizing the triggering enhanced input action before all else
+		if (CommonUI::IsEnhancedInputSupportEnabled() && TriggeringEnhancedInputAction)
+		{
+			InputActionWidget->SetEnhancedInputAction(TriggeringEnhancedInputAction);
+		}
+		// Prefer visualizing the triggering input action next
+		else if (!TriggeringInputAction.IsNull())
 		{
 			InputActionWidget->SetInputAction(TriggeringInputAction);
 		}
@@ -1038,11 +1090,28 @@ void UCommonButtonBase::UpdateInputActionWidget()
 		else if (bShouldUseFallbackDefaultInputAction && bButtonEnabled)
 		{
 			FDataTableRowHandle HoverStateHandle;
+			UInputAction* HoverEnhancedInputAction = nullptr;
+			bool bIsEnhancedInputSupportEnabled = CommonUI::IsEnhancedInputSupportEnabled();
 			if (IsHovered())
 			{
-				HoverStateHandle = ICommonInputModule::GetSettings().GetDefaultClickAction();
+				if (bIsEnhancedInputSupportEnabled)
+				{
+					HoverEnhancedInputAction = ICommonInputModule::GetSettings().GetEnhancedInputClickAction();
+				}
+				else
+				{
+					HoverStateHandle = ICommonInputModule::GetSettings().GetDefaultClickAction();
+				}
+			}
+
+			if (bIsEnhancedInputSupportEnabled)
+			{
+				InputActionWidget->SetEnhancedInputAction(HoverEnhancedInputAction);
+			}
+			else
+			{
+				InputActionWidget->SetInputAction(HoverStateHandle);
 			}
-			InputActionWidget->SetInputAction(HoverStateHandle);
 		}
 		else
 		{
diff --git a/Engine/Plugins/Experimental/CommonUI/Source/CommonUI/Private/CommonTabListWidgetBase.cpp b/Engine/Plugins/Experimental/CommonUI/Source/CommonUI/Private/CommonTabListWidgetBase.cpp
--- a/Engine/Plugins/Experimental/CommonUI/Source/CommonUI/Private/CommonTabListWidgetBase.cpp
+++ b/Engine/Plugins/Experimental/CommonUI/Source/CommonUI/Private/CommonTabListWidgetBase.cpp
@@ -5,6 +5,8 @@
 #include "CommonAnimatedSwitcher.h"
 #include "Groups/CommonButtonGroupBase.h"
 #include "Input/CommonUIInputTypes.h"
+#include "CommonUITypes.h"
+#include "InputAction.h"
 
 #include UE_INLINE_GENERATED_CPP_BY_NAME(CommonTabListWidgetBase)
 
@@ -190,8 +192,24 @@ void UCommonTabListWidgetBase::UpdateBindings()
 	// New input system binding flow
 	if (bIsListeningForInput)
 	{
-		NextTabActionHandle = RegisterUIActionBinding(FBindUIActionArgs(NextTabInputActionData, false, FSimpleDelegate::CreateUObject(this, &UCommonTabListWidgetBase::HandleNextTabAction)));
-		PrevTabActionHandle = RegisterUIActionBinding(FBindUIActionArgs(PreviousTabInputActionData, false, FSimpleDelegate::CreateUObject(this, &UCommonTabListWidgetBase::HandlePreviousTabAction)));
+		bool bIsEnhancedInputSupportEnabled = CommonUI::IsEnhancedInputSupportEnabled();
+		if (bIsEnhancedInputSupportEnabled && NextTabEnhancedInputAction)
+		{
+			NextTabActionHandle = RegisterUIActionBinding(FBindUIActionArgs(NextTabEnhancedInputAction, false, FSimpleDelegate::CreateUObject(this, &UCommonTabListWidgetBase::HandleNextTabAction)));
+		}
+		else
+		{
+			NextTabActionHandle = RegisterUIActionBinding(FBindUIActionArgs(NextTabInputActionData, false, FSimpleDelegate::CreateUObject(this, &UCommonTabListWidgetBase::HandleNextTabAction)));
+		}
+
+		if (bIsEnhancedInputSupportEnabled && PreviousTabEnhancedInputAction)
+		{
+			PrevTabActionHandle = RegisterUIActionBinding(FBindUIActionArgs(PreviousTabEnhancedInputAction, false, FSimpleDelegate::CreateUObject(this, &UCommonTabListWidgetBase::HandlePreviousTabAction)));
+		}
+		else
+		{
+			PrevTabActionHandle = RegisterUIActionBinding(FBindUIActionArgs(PreviousTabInputActionData, false, FSimpleDelegate::CreateUObject(this, &UCommonTabListWidgetBase::HandlePreviousTabAction)));
+		}
 	}
 	else
 	{
diff --git a/Engine/Plugins/Experimental/CommonUI/Source/CommonUI/Private/CommonUISubsystemBase.cpp b/Engine/Plugins/Experimental/CommonUI/Source/CommonUI/Private/CommonUISubsystemBase.cpp
--- a/Engine/Plugins/Experimental/CommonUI/Source/CommonUI/Private/CommonUISubsystemBase.cpp
+++ b/Engine/Plugins/Experimental/CommonUI/Source/CommonUI/Private/CommonUISubsystemBase.cpp
@@ -7,8 +7,10 @@
 #include "CommonInputSubsystem.h"
 #include "CommonUITypes.h"
 #include "IAnalyticsProviderET.h"
+#include "InputAction.h"
 #include "Components/Widget.h"
 #include "Engine/GameInstance.h"
+#include "EnhancedInputSubsystems.h"
 #include "Styling/StyleDefaults.h"
 
 #include UE_INLINE_GENERATED_CPP_BY_NAME(CommonUISubsystemBase)
@@ -58,6 +60,16 @@ FSlateBrush UCommonUISubsystemBase::GetInputActionButtonIcon(const FDataTableRow
 	return *FStyleDefaults::GetNoBrush();
 }
 
+FSlateBrush UCommonUISubsystemBase::GetEnhancedInputActionButtonIcon(const UInputAction* InputAction, const ULocalPlayer* LocalPlayer) const
+{
+	if (LocalPlayer)
+	{
+		return CommonUI::GetIconForEnhancedInputAction(UCommonInputSubsystem::Get(LocalPlayer), InputAction);
+	}
+
+	return *FStyleDefaults::GetNoBrush();
+}
+
 /**
  * @EventName CommonUI.ButtonClicked
  * @Trigger (DISABLED) Button presses that are marked for analytic events
diff --git a/Engine/Plugins/Experimental/CommonUI/Source/CommonUI/Private/CommonUITypes.cpp b/Engine/Plugins/Experimental/CommonUI/Source/CommonUI/Private/CommonUITypes.cpp
--- a/Engine/Plugins/Experimental/CommonUI/Source/CommonUI/Private/CommonUITypes.cpp
+++ b/Engine/Plugins/Experimental/CommonUI/Source/CommonUI/Private/CommonUITypes.cpp
@@ -5,11 +5,14 @@
 #include "CommonInputBaseTypes.h"
 #include "CommonInputSubsystem.h"
 #include "CommonUIPrivate.h"
+#include "EnhancedActionKeyMapping.h"
 #include "EnhancedInputSubsystems.h"
 #include "Engine/LocalPlayer.h"
 #include "HAL/PlatformInput.h"
 #include "ICommonInputModule.h"
 #include "InputAction.h"
+#include "InputActionValue.h"
+#include "PlayerMappableKeySettings.h"
 #include "Styling/SlateTypes.h"
 #include "Styling/StyleDefaults.h"
 
@@ -191,6 +194,16 @@ void FCommonInputActionDataBase::AddGamepadInputOverride(const FName& GamepadNam
 	GamepadInputOverrides.Add(GamepadName, InputInfo);
 }
 
+const UCommonInputMetadata* UCommonMappingContextMetadata::GetCommonInputMetadata(const UInputAction* InInputAction) const
+{
+	if (const TObjectPtr<const UCommonInputMetadata>* PerActionMetdata = PerActionEnhancedInputMetadata.Find(InInputAction))
+	{
+		return *PerActionMetdata;
+	}
+
+	return EnhancedInputMetadata;
+}
+
 void CommonUI::SetupStyles()
 {
 	EmptyScrollBoxStyle.BottomShadowBrush.DrawAs = ESlateBrushDrawType::NoDrawType;
@@ -235,7 +248,47 @@ bool CommonUI::IsEnhancedInputSupportEnabled()
 	return bEnabled;
 }
 
-FSlateBrush CommonUI::GetIconForEnhancedInputAction(const UCommonInputSubsystem* CommonInputSubsystem, const TObjectPtr<UInputAction> InputAction)
+TObjectPtr<const UCommonInputMetadata> CommonUI::GetEnhancedInputActionMetadata(const ULocalPlayer* LocalPlayer, const UInputAction* InputAction)
+{
+	if (!InputAction->GetPlayerMappableKeySettings())
+	{
+		return nullptr;
+	}
+
+	if (UBlueprint* MetadataBP = Cast<UBlueprint>(InputAction->GetPlayerMappableKeySettings()->Metadata))
+	{
+		if (ICommonMappingContextMetadataInterface* MappingContextMetadata = Cast<ICommonMappingContextMetadataInterface>(MetadataBP->GeneratedClass->GetDefaultObject()))
+		{
+			return MappingContextMetadata->GetCommonInputMetadata(InputAction);
+		}
+	}
+
+	return nullptr;
+}
+
+void CommonUI::GetEnhancedInputActionKeys(const ULocalPlayer* LocalPlayer, const UInputAction* InputAction, TArray<FKey>& OutKeys)
+{
+	if (LocalPlayer)
+	{
+		if (UEnhancedInputLocalPlayerSubsystem* EnhancedInputLocalPlayerSubsystem = LocalPlayer->GetSubsystem<UEnhancedInputLocalPlayerSubsystem>())
+		{
+			OutKeys = EnhancedInputLocalPlayerSubsystem->QueryKeysMappedToAction(InputAction);
+		}
+	}
+}
+
+void CommonUI::InjectEnhancedInputForAction(const ULocalPlayer* LocalPlayer, const UInputAction* InputAction, FInputActionValue RawValue)
+{
+	if (LocalPlayer)
+	{
+		if (UEnhancedInputLocalPlayerSubsystem* EnhancedInputLocalPlayerSubsystem = LocalPlayer->GetSubsystem<UEnhancedInputLocalPlayerSubsystem>())
+		{
+			EnhancedInputLocalPlayerSubsystem->InjectInputForAction(InputAction, RawValue, {}, {});
+		}
+	}
+}
+
+FSlateBrush CommonUI::GetIconForEnhancedInputAction(const UCommonInputSubsystem* CommonInputSubsystem, const UInputAction* InputAction)
 {
 	TArray<FKey> Keys;
 	if (ULocalPlayer* LocalPlayer = CommonInputSubsystem->GetLocalPlayer())
@@ -253,4 +306,4 @@ FSlateBrush CommonUI::GetIconForEnhancedInputAction(const UCommonInputSubsystem*
 	}
 
 	return *FStyleDefaults::GetNoBrush();
-}
\ No newline at end of file
+}
diff --git a/Engine/Plugins/Experimental/CommonUI/Source/CommonUI/Private/Input/UIActionRouterTypes.cpp b/Engine/Plugins/Experimental/CommonUI/Source/CommonUI/Private/Input/UIActionRouterTypes.cpp
--- a/Engine/Plugins/Experimental/CommonUI/Source/CommonUI/Private/Input/UIActionRouterTypes.cpp
+++ b/Engine/Plugins/Experimental/CommonUI/Source/CommonUI/Private/Input/UIActionRouterTypes.cpp
@@ -5,6 +5,8 @@
 #include "Input/CommonUIActionRouterBase.h"
 #include "CommonActivatableWidget.h"
 #include "Input/CommonUIInputTypes.h"
+#include "InputAction.h"
+#include "InputActionValue.h"
 #include "CommonInputSettings.h"
 #include "CommonInputSubsystem.h"
 #include "ICommonInputModule.h"
@@ -112,6 +114,7 @@ FUIActionBinding::FUIActionBinding(const UWidget& InBoundWidget, const FBindUIAc
 	, OnExecuteAction(BindArgs.OnExecuteAction)
 	, Handle(IdCounter++)
 	, LegacyActionTableRow(BindArgs.LegacyActionTableRow)
+	, InputAction(BindArgs.InputAction)
 {
 	OnHoldActionProgressed.Add(BindArgs.OnHoldActionProgressed);
 
@@ -146,6 +149,11 @@ FUIActionBinding::FUIActionBinding(const UWidget& InBoundWidget, const FBindUIAc
 			RegisterKeyMappingFunc(KeyMapping);
 		}
 	}
+	else if (CommonUI::IsEnhancedInputSupportEnabled() && BindArgs.InputAction.IsValid())
+	{
+		// Nothing else to do if we have an enhanced input action,
+		// the input action itself will be queried against for keys later
+	}
 	else
 	{
 		const FCommonInputActionDataBase* LegacyActionData = GetLegacyInputActionData();
@@ -182,7 +190,8 @@ FUIActionBinding::FUIActionBinding(const UWidget& InBoundWidget, const FBindUIAc
 
 FUIActionBindingHandle FUIActionBinding::TryCreate(const UWidget& InBoundWidget, const FBindUIActionArgs& BindArgs)
 {
-	if (BindArgs.GetActionName().IsNone())
+	bool bIsEnhancedInputSupportEnabled = CommonUI::IsEnhancedInputSupportEnabled();
+	if (BindArgs.GetActionName().IsNone() && (!bIsEnhancedInputSupportEnabled || !BindArgs.InputAction.IsValid()))
 	{
 		UE_LOG(LogUIActionRouter, Error, TEXT("Cannot create action binding for widget [%s] - no action provided."), *InBoundWidget.GetName());
 		return FUIActionBindingHandle();
@@ -202,6 +211,11 @@ FUIActionBindingHandle FUIActionBinding::TryCreate(const UWidget& InBoundWidget,
 		UE_LOG(LogUIActionRouter, Error, TEXT("Cannot bind widget [%s] to action [%s] - provided legacy data table row does not resolve to valid data."), *InBoundWidget.GetName(), *BindArgs.GetActionName().ToString());
 		return FUIActionBindingHandle();
 	}
+	else if (bIsEnhancedInputSupportEnabled && !BindArgs.InputAction.IsValid())
+	{
+		UE_LOG(LogUIActionRouter, Error, TEXT("Cannot bind widget [%s] to action [%s] - provided input action is invalid."), *InBoundWidget.GetName(), *BindArgs.GetActionName().ToString());
+		return FUIActionBindingHandle();
+	}
 	
 	// Make sure there is no existing binding for the same action associated with the same widget
 	for (const TPair<FUIActionBindingHandle, TSharedPtr<FUIActionBinding>>& HandleBindingPair : AllRegistrationsByHandle)
@@ -524,20 +538,68 @@ bool FActionRouterBindingCollection::ProcessNormalInput(ECommonInputMode ActiveI
 		{
 			if (ActiveInputMode == ECommonInputMode::All || ActiveInputMode == Binding->InputMode)
 			{
-				for (const FUIActionKeyMapping& KeyMapping : Binding->NormalMappings)
+				auto TryConsumeInput = [&](const FKey& InKey, const UInputAction* InInputAction)
 				{
 					// A persistent displayed action skips the normal rules for reachability, since it'll always appear in a bound action bar
 					const bool bIsDisplayedPersistentAction = Binding->bIsPersistent && Binding->bDisplayInActionBar;
-					if (KeyMapping.Key == Key && Binding->InputEvent == InputEvent && (bIsDisplayedPersistentAction || IsWidgetReachableForInput(Binding->BoundWidget.Get())))
+					if (InKey == Key && Binding->InputEvent == InputEvent && (bIsDisplayedPersistentAction || IsWidgetReachableForInput(Binding->BoundWidget.Get())))
 					{
 						// Just in case this was in the middle of a hold process with a different key, reset now
 						Binding->CancelHold();
-						Binding->OnExecuteAction.ExecuteIfBound();
+
+						// If injecting enhanced input. don't fire 'OnExecuteAction' since that can be manually done if desired in BP
+						bool bEnhancedInputInjected = false;
+						if (InInputAction)
+						{
+							if (TObjectPtr<const UCommonInputMetadata> CommonInputMetadata = CommonUI::GetEnhancedInputActionMetadata(GetActionRouter().GetLocalPlayerChecked(), InInputAction))
+							{
+								// Non generic actions should inject enhanced input so users can bind to enhanced input events
+								if (!CommonInputMetadata->bIsGenericInputAction)
+								{
+									FInputActionValue RawValue = FInputActionValue(true);
+									CommonUI::InjectEnhancedInputForAction(GetActionRouter().GetLocalPlayerChecked(), InInputAction, RawValue);
+									bEnhancedInputInjected = true;
+								}
+							}
+						}
+						
+						if (!bEnhancedInputInjected)
+						{
+							Binding->OnExecuteAction.ExecuteIfBound();
+						}
+
 						if (Binding->bConsumesInput)
 						{
 							return true;
 						}
 					}
+
+					return false;
+				};
+
+
+				if (CommonUI::IsEnhancedInputSupportEnabled() && Binding->InputAction.IsValid())
+				{
+					if (const UInputAction* InputAction = Binding->InputAction.Get())
+					{
+						TArray<FKey> InputActionKeys;
+						CommonUI::GetEnhancedInputActionKeys(GetActionRouter().GetLocalPlayerChecked(), InputAction, InputActionKeys);
+						for (const FKey& InputActionKey : InputActionKeys)
+						{
+							if (TryConsumeInput(InputActionKey, InputAction))
+							{
+								return true;
+							}
+						}
+					}
+				}
+
+				for (const FUIActionKeyMapping& KeyMapping : Binding->NormalMappings)
+				{
+					if (TryConsumeInput(KeyMapping.Key, nullptr))
+					{
+						return true;
+					}
 				}
 			}
 		}
diff --git a/Engine/Plugins/Experimental/CommonUI/Source/CommonUI/Public/CommonActionHandlerInterface.h b/Engine/Plugins/Experimental/CommonUI/Source/CommonUI/Public/CommonActionHandlerInterface.h
--- a/Engine/Plugins/Experimental/CommonUI/Source/CommonUI/Public/CommonActionHandlerInterface.h
+++ b/Engine/Plugins/Experimental/CommonUI/Source/CommonUI/Public/CommonActionHandlerInterface.h
@@ -8,6 +8,8 @@
 #include "UObject/Interface.h"
 #include "CommonActionHandlerInterface.generated.h"
 
+// @FIXME: This entire file is related to legacy CommonUI and should be removed in 5.3 - UE-164871
+
 enum EInputEvent : int;
 namespace ETouchType { enum Type : int; }
 struct FKey;
@@ -39,6 +41,7 @@ DECLARE_DELEGATE_OneParam(FCommonActionProgressSingleNative, float);
 DECLARE_DYNAMIC_MULTICAST_DELEGATE_OneParam(FCommonActionProgress, float, HeldPercent);
 DECLARE_MULTICAST_DELEGATE_OneParam(FCommonActionProgressNative, float);
 
+// @FIXME: This struct is related to legacy CommonUI and should be removed in 5.3 - UE-164871
 USTRUCT(BlueprintType)
 struct COMMONUI_API FCommonInputActionHandlerData
 {
@@ -110,6 +113,8 @@ struct FCommonInputActionHandlerDelegateData
 /**
  *  Action Handler Interface is primarily used to take key input and do something with it
  *	in the implementation of the interface or another user widget.
+ * 
+ *  @FIXME: This interface is related to legacy CommonUI and should be removed in 5.3 - UE-164871
  */
 
 UINTERFACE()
diff --git a/Engine/Plugins/Experimental/CommonUI/Source/CommonUI/Public/CommonActivatableWidget.h b/Engine/Plugins/Experimental/CommonUI/Source/CommonUI/Public/CommonActivatableWidget.h
--- a/Engine/Plugins/Experimental/CommonUI/Source/CommonUI/Public/CommonActivatableWidget.h
+++ b/Engine/Plugins/Experimental/CommonUI/Source/CommonUI/Public/CommonActivatableWidget.h
@@ -8,6 +8,7 @@
 
 class FActivatableTreeNode;
 class UCommonInputActionDomain;
+class UInputMappingContext;
 
 DECLARE_DYNAMIC_MULTICAST_DELEGATE(FOnWidgetActivationChanged);
 
@@ -186,6 +187,14 @@ protected:
 
 	UPROPERTY(EditAnywhere, BlueprintReadOnly, Category = "Input", meta = (InlineEditConditionToggle))
 	bool bOverrideActionDomain = false;
+	
+	/** Optional mapping context to be applied & removed on activation & deactivation respectfully. */
+	UPROPERTY(EditAnywhere, Category="Input", meta = (EditCondition = "CommonInput.CommonInputSettings.IsEnhancedInputSupportEnabled", EditConditionHides))
+	TObjectPtr<UInputMappingContext> InputMapping;
+
+	/** Enhanced Input priority. Higher priority input mappings will be prioritized over mappings with a lower priority. */
+	UPROPERTY(EditAnywhere, Category="Input", meta = (EditCondition = "CommonInput.CommonInputSettings.IsEnhancedInputSupportEnabled", EditConditionHides))
+	int32 InputMappingPriority = 0;
 
 	/**
 	 * Enable to override the inherited ActionDomain from owning CommonActivatableWidget.
diff --git a/Engine/Plugins/Experimental/CommonUI/Source/CommonUI/Public/CommonButtonBase.h b/Engine/Plugins/Experimental/CommonUI/Source/CommonUI/Public/CommonButtonBase.h
--- a/Engine/Plugins/Experimental/CommonUI/Source/CommonUI/Public/CommonButtonBase.h
+++ b/Engine/Plugins/Experimental/CommonUI/Source/CommonUI/Public/CommonButtonBase.h
@@ -16,6 +16,7 @@ enum class ECommonInputType : uint8;
 class UCommonTextStyle;
 class SBox;
 class UCommonActionWidget;
+class UInputAction;
 
 class COMMONUI_API FCommonButtonMetaData : public ISlateMetaData
 {
@@ -415,10 +416,18 @@ public:
 	UFUNCTION(BlueprintCallable, Category = "Common Button|Setters")
 	void SetTriggeringInputAction(const FDataTableRowHandle & InputActionRow);
 
+	/** Updates the current triggering enhanced input action, requires enhanced input enabled in CommonUI settings */
+	UFUNCTION(BlueprintCallable, Category = "Common Button|Setters")
+	void SetTriggeringEnhancedInputAction(UInputAction* InInputAction);
+
 	/** Gets the appropriate input action that is set */
 	UFUNCTION(BlueprintCallable, Category = "Common Button|Getters")
 	bool GetInputAction(FDataTableRowHandle &InputActionRow) const;
 
+	/** Gets the appropriate enhanced input action that is set */
+	UFUNCTION(BlueprintCallable, Category = "Common Button|Getters")
+	UInputAction* GetEnhancedInputAction() const;
+
 	/** Updates the bIsFocusable flag */
 	UFUNCTION(BlueprintCallable, Category = "Common Button|Getters")
 	void SetIsFocusable(bool bInIsFocusable);
@@ -604,18 +613,26 @@ protected:
 	UFUNCTION(BlueprintCallable, meta=(BlueprintProtected="true"), Category = "Common Button")
 	void SetSelectedInternal(bool bInSelected, bool bAllowSound = true, bool bBroadcast = true);
 
+	/** Callback fired when input action datatable row changes */
 	UFUNCTION(BlueprintImplementableEvent, Category = "Common Button")
 	void OnTriggeredInputActionChanged(const FDataTableRowHandle& NewTriggeredAction);
 
+	/** Callback fired when triggered input action datatable row changes */
 	UFUNCTION(BlueprintImplementableEvent, Category = "Common Button")
 	void OnTriggeringInputActionChanged(const FDataTableRowHandle& NewTriggeredAction);
 
+	/** Callback fired when enhanced input action changes */
+	UFUNCTION(BlueprintImplementableEvent, Category = "Common Button")
+	void OnTriggeringEnhancedInputActionChanged(const UInputAction* InInputAction);
+
+	/** Callback fired continously during hold interactions */
 	UFUNCTION(BlueprintImplementableEvent, Category = "Common Button")
 	void OnActionProgress(float HeldPercent);
 	
 	UFUNCTION()
 	virtual void NativeOnActionProgress(float HeldPercent);
 
+	/** Callback fired when hold events complete */
 	UFUNCTION(BlueprintImplementableEvent, Category = "Common Button")
 	void OnActionComplete();
 
@@ -739,7 +756,12 @@ public:
 	UPROPERTY(EditAnywhere, BlueprintReadOnly, Category = Input)
 	TEnumAsByte<EButtonPressMethod::Type> PressMethod;
 
-	/** This is the priority for the TriggeringInputAction.  The first, HIGHEST PRIORITY widget will handle the input action, and no other widgets will be considered.  Additionally, no inputs with a priority below the current ActivatablePanel's Input Priority value will even be considered! */
+	/** 
+	 * This is the priority for the TriggeringInputAction.  The first, HIGHEST PRIORITY widget will handle the input action, and no other widgets will be considered.  
+	 * Additionally, no inputs with a priority below the current ActivatablePanel's Input Priority value will even be considered! 
+	 * 
+	 * @TODO: This is part of legacy CommonUI and should be removed
+	 */
 	UPROPERTY(EditAnywhere, BlueprintReadOnly, Category = Input, meta = (ExposeOnSpawn = true))
 	int32 InputPriority;
 
@@ -750,9 +772,18 @@ public:
 	UPROPERTY(EditAnywhere, BlueprintReadOnly, Category = Input, meta = (ExposeOnSpawn = true, RowType = "/Script/CommonUI.CommonInputActionDataBase"))
 	FDataTableRowHandle TriggeringInputAction;
 
+	/** 
+	 *	The enhanced input action that is bound to this button. The common input manager will trigger this button to 
+	 *	click if the action was pressed 
+	 */
+	UPROPERTY(EditAnywhere, BlueprintReadOnly, Category = Input, meta = (EditCondition = "CommonInput.CommonInputSettings.IsEnhancedInputSupportEnabled", EditConditionHides))
+	TObjectPtr<UInputAction> TriggeringEnhancedInputAction;
+
 	/**
-	 *	The input action that can be visualized as well as triggered when the user
-	 *	clicks the button.
+	 * The input action that can be visualized as well as triggered when the user
+	 * clicks the button.
+	 * 
+	 * @TODO: This is part of legacy CommonUI and should be removed
 	 */
 	FDataTableRowHandle TriggeredInputAction;
 
@@ -856,6 +887,7 @@ protected:
 	/**
 	 * Optionally bound widget for visualization behavior of an input action;
 	 * NOTE: If specified, will visualize according to the following algorithm:
+	 * If TriggeringEnhancedInputAction is specified, visualize it else:
 	 * If TriggeringInputAction is specified, visualize it else:
 	 * If TriggeredInputAction is specified, visualize it else:
 	 * Visualize the default click action while hovered
diff --git a/Engine/Plugins/Experimental/CommonUI/Source/CommonUI/Public/CommonTabListWidgetBase.h b/Engine/Plugins/Experimental/CommonUI/Source/CommonUI/Public/CommonTabListWidgetBase.h
--- a/Engine/Plugins/Experimental/CommonUI/Source/CommonUI/Public/CommonTabListWidgetBase.h
+++ b/Engine/Plugins/Experimental/CommonUI/Source/CommonUI/Public/CommonTabListWidgetBase.h
@@ -11,6 +11,7 @@ class UCommonAnimatedSwitcher;
 
 class UCommonButtonBase;
 class UCommonButtonGroupBase;
+class UInputAction;
 
 /** Information about a registered tab in the tab list */
 USTRUCT()
@@ -179,6 +180,14 @@ protected:
 	UPROPERTY(EditAnywhere, BlueprintReadOnly, Category = TabList, meta = (RowType = "/Script/CommonUI.CommonInputActionDataBase"))
 	FDataTableRowHandle PreviousTabInputActionData;
 
+	/** The input action to listen for causing the next tab to be selected */
+	UPROPERTY(EditAnywhere, BlueprintReadOnly, Category = TabList, meta = (EditCondition = "CommonInput.CommonInputSettings.IsEnhancedInputSupportEnabled", EditConditionHides))
+	TObjectPtr<UInputAction> NextTabEnhancedInputAction;
+
+	/** The input action to listen for causing the previous tab to be selected */
+	UPROPERTY(EditAnywhere, BlueprintReadOnly, Category = TabList, meta = (EditCondition = "CommonInput.CommonInputSettings.IsEnhancedInputSupportEnabled", EditConditionHides))
+	TObjectPtr<UInputAction> PreviousTabEnhancedInputAction;
+
 	/** Whether to register to handle tab list input immediately upon construction */
 	UPROPERTY(EditAnywhere, BlueprintReadOnly, Category = TabList, meta = (ExposeOnSpawn = "true"))
 	bool bAutoListenForInput;
diff --git a/Engine/Plugins/Experimental/CommonUI/Source/CommonUI/Public/CommonUISubsystemBase.h b/Engine/Plugins/Experimental/CommonUI/Source/CommonUI/Public/CommonUISubsystemBase.h
--- a/Engine/Plugins/Experimental/CommonUI/Source/CommonUI/Public/CommonUISubsystemBase.h
+++ b/Engine/Plugins/Experimental/CommonUI/Source/CommonUI/Public/CommonUISubsystemBase.h
@@ -13,6 +13,7 @@ struct FSlateBrush;
 class IAnalyticsProviderET;
 class UWidget;
 class ULocalPlayer;
+class UInputAction;
 
 UCLASS(DisplayName = "CommonUI")
 class COMMONUI_API UCommonUISubsystemBase : public UGameInstanceSubsystem
@@ -32,6 +33,10 @@ public:
 	UFUNCTION(BlueprintCallable, Category = CommonUISubsystem)
 	FSlateBrush GetInputActionButtonIcon(const FDataTableRowHandle& InputActionRowHandle, ECommonInputType InputType, const FName& GamepadName) const;
 
+	// Gets Action Button Icon for given action and player, enhanced input API currently does not allow input type specification
+	UFUNCTION(BlueprintCallable, Category = CommonUISubsystem)
+	FSlateBrush GetEnhancedInputActionButtonIcon(const UInputAction* InputAction, const ULocalPlayer* LocalPlayer) const;
+
 	/** Analytic Events **/
 
 	//CommonUI.ButtonClicked
diff --git a/Engine/Plugins/Experimental/CommonUI/Source/CommonUI/Public/CommonUITypes.h b/Engine/Plugins/Experimental/CommonUI/Source/CommonUI/Public/CommonUITypes.h
--- a/Engine/Plugins/Experimental/CommonUI/Source/CommonUI/Public/CommonUITypes.h
+++ b/Engine/Plugins/Experimental/CommonUI/Source/CommonUI/Public/CommonUITypes.h
@@ -12,6 +12,8 @@ enum class ECommonInputType : uint8;
 struct FScrollBoxStyle;
 
 class UCommonInputSubsystem;
+class UInputAction;
+struct FInputActionValue;
 
 UENUM(BlueprintType)
 enum class EInputActionState : uint8
@@ -166,16 +168,113 @@ public:
 	}
 };
 
+/** 
+ * Metadata CommonUI will try to acquire from Enhanced Input Mapping Contexts (IMC)
+ * 
+ * You can Inherit from this class if you have any info that may need to be parsed per platform
+ * by CommonUI. IMC's can be specified per platform, so each platform may have different
+ * Common Input Metadata
+ * 
+ * Note: We intentionally do not define any context-independant metadata. Even though some
+ * metadata should be context-independant (Like NavBarPriority below). Locking it that info
+ * to a seperate metadata type prevents any chance of future overriding. Instead, we prefer
+ * info for all metadata to be set across all instances.
+ */
+UCLASS(Blueprintable, EditInlineNew, CollapseCategories)
+class COMMONUI_API UCommonInputMetadata : public UObject
+{
+	GENERATED_BODY()
+
+public:
+
+	/** Priority in nav-bar */
+	UPROPERTY(EditAnywhere, BlueprintReadOnly, Category = "CommonInput")
+	int32 NavBarPriority = 0;
+
+	/** 
+	 * Generic actions like accept or face button top will be subscribed to by multiple
+	 * UI elements. These actions will not broadcast enhanced input action delegates
+	 * such as "Triggered, Ongoing, Canceled, or Completed." Since those delegates
+	 * would be fired by multiple UI elements.
+	 * 
+	 * Non-generic input actions will fire Enhanced Input events. However they will 
+	 * not fire CommonUI action bindings (Since those can be manually fired in BP).
+	 */
+	UPROPERTY(EditAnywhere, BlueprintReadOnly, Category = "CommonInput")
+	bool bIsGenericInputAction = true;
+};
+
+/**
+ * Interface for metadata objects specified in Enhanced Input Mapping Contexts (IMC).
+ * 
+ * We provide an interface since it's possible you may need the IMC metadata for 
+ * non-CommonUI info. In this scenario you can implement this interface and CommonUI
+ * will still be able to gather info it needs to function correctly with your 
+ * Enhanced Input Actions / IMC's.
+ * 
+ * If you don't have any metadata needs or your UI IMC's are for CommonUI only, 
+ * then you should use the provided 'UCommonMappingContextMetadata' below.
+ */
+UINTERFACE()
+class COMMONUI_API UCommonMappingContextMetadataInterface : public UInterface
+{
+	GENERATED_BODY()
+};
+
+class COMMONUI_API ICommonMappingContextMetadataInterface
+{
+	GENERATED_BODY()
+
+public:
+	/** 
+	 * Gets base info needed from CommonUI from this IMC metadata 
+	 * Accepts InputAction as an arg to allow for user to create one metadata with multiple
+	 * values per action, rather than having to create one metadata asset per unique value
+	 * 
+	 * @param InputAction the input action that may or may not be used to help determine correct metadata
+	 */
+	virtual const UCommonInputMetadata* GetCommonInputMetadata(const UInputAction* InInputAction) const = 0;
+};
+
+/**
+ * Base CommonUI metadata implementation for specification in IMC's.
+ * 
+ * Utilizes a map of input actions to metadata to prevent users from having to create
+ * multiple metadata assets / instances. Using this map is not mandatory.
+ */
+UCLASS(Blueprintable)
+class COMMONUI_API UCommonMappingContextMetadata : public UObject, public ICommonMappingContextMetadataInterface
+{
+	GENERATED_BODY()
+
+public:
+	/** Fallback or default metadata CommonUI relies on if no per-action meta is found below */
+	UPROPERTY(EditAnywhere, Instanced, BlueprintReadOnly, Category = "CommonInput")
+	TObjectPtr<UCommonInputMetadata> EnhancedInputMetadata;
+
+	/** Map of action to metadata, allows creation of single metadata asset rather than one per input action type */
+	UPROPERTY(EditAnywhere, Instanced, BlueprintReadOnly, Category = "CommonInput")
+	TMap<TObjectPtr<UInputAction>, TObjectPtr<const UCommonInputMetadata>> PerActionEnhancedInputMetadata;
+
+public:
+	virtual const UCommonInputMetadata* GetCommonInputMetadata(const UInputAction* InInputAction) const override;
+};
+
 class COMMONUI_API CommonUI
 {
 public:
 	static void SetupStyles();
 	static FScrollBoxStyle EmptyScrollBoxStyle;
+
 	static const FCommonInputActionDataBase* GetInputActionData(const FDataTableRowHandle& InputActionRowHandle);
 	static FSlateBrush GetIconForInputActions(const UCommonInputSubsystem* CommonInputSubsystem, const TArray<FDataTableRowHandle>& InputActions);
 
 	static bool IsEnhancedInputSupportEnabled();
-	static FSlateBrush GetIconForEnhancedInputAction(const UCommonInputSubsystem* CommonInputSubsystem, const TObjectPtr<class UInputAction> InputAction);
+
+	static TObjectPtr<const UCommonInputMetadata> GetEnhancedInputActionMetadata(const ULocalPlayer* LocalPlayer, const UInputAction* InputAction);
+	static void GetEnhancedInputActionKeys(const ULocalPlayer* LocalPlayer, const UInputAction* InputAction, TArray<FKey>& OutKeys);
+	static void InjectEnhancedInputForAction(const ULocalPlayer* LocalPlayer, const UInputAction* InputAction, FInputActionValue RawValue);
+	static FSlateBrush GetIconForEnhancedInputAction(const UCommonInputSubsystem* CommonInputSubsystem, const UInputAction* InputAction);
 };
 
 DECLARE_DYNAMIC_DELEGATE_OneParam(FOnItemClicked, UUserWidget*, Widget);
diff --git a/Engine/Plugins/Experimental/CommonUI/Source/CommonUI/Public/Input/CommonUIInputTypes.h b/Engine/Plugins/Experimental/CommonUI/Source/CommonUI/Public/Input/CommonUIInputTypes.h
--- a/Engine/Plugins/Experimental/CommonUI/Source/CommonUI/Public/Input/CommonUIInputTypes.h
+++ b/Engine/Plugins/Experimental/CommonUI/Source/CommonUI/Public/Input/CommonUIInputTypes.h
@@ -6,6 +6,7 @@
 #include "UITag.h"
 #include "CommonInputModeTypes.h"
 #include "Engine/EngineBaseTypes.h"
+#include "InputAction.h"
 
 struct COMMONUI_API FBindUIActionArgs
 {
@@ -20,26 +21,41 @@ struct COMMONUI_API FBindUIActionArgs
 		, OnExecuteAction(InOnExecuteAction)
 	{}
 
-	// @TODO: DarenC - Remove legacy.
+	// @TODO: Rename non-legacy in 5.3. We no longer have any active plans to remove data tables in CommonUI.
 	FBindUIActionArgs(const FDataTableRowHandle& InLegacyActionTableRow, const FSimpleDelegate& InOnExecuteAction)
 		: LegacyActionTableRow(InLegacyActionTableRow)
 		, OnExecuteAction(InOnExecuteAction)
 	{}
 
-	// @TODO: DarenC - Remove legacy.
+	// @TODO: Rename non-legacy in 5.3. We no longer have any active plans to remove data tables in CommonUI.
 	FBindUIActionArgs(const FDataTableRowHandle& InLegacyActionTableRow, bool bShouldDisplayInActionBar, const FSimpleDelegate& InOnExecuteAction)
 		: LegacyActionTableRow(InLegacyActionTableRow)
 		, bDisplayInActionBar(bShouldDisplayInActionBar)
 		, OnExecuteAction(InOnExecuteAction)
 	{}
 
+	FBindUIActionArgs(const UInputAction* InInputAction, const FSimpleDelegate & InOnExecuteAction)
+		: InputAction(InInputAction)
+		, OnExecuteAction(InOnExecuteAction)
+	{}
+
+	FBindUIActionArgs(const UInputAction* InInputAction, bool bShouldDisplayInActionBar, const FSimpleDelegate & InOnExecuteAction)
+		: InputAction(InInputAction)
+		, bDisplayInActionBar(bShouldDisplayInActionBar)
+		, OnExecuteAction(InOnExecuteAction)
+	{}
+
 	FName GetActionName() const;
 
 	bool ActionHasHoldMappings() const;
 
 	FUIActionTag ActionTag;
+
+	// @TODO: Rename non-legacy in 5.3. We no longer have any active plans to remove data tables in CommonUI.
 	FDataTableRowHandle LegacyActionTableRow;
 
+	TWeakObjectPtr<const UInputAction> InputAction;
+
 	ECommonInputMode InputMode = ECommonInputMode::Menu;
 	EInputEvent KeyEvent = IE_Pressed;
 
diff --git a/Engine/Plugins/Experimental/CommonUI/Source/CommonUI/Public/Input/UIActionBinding.h b/Engine/Plugins/Experimental/CommonUI/Source/CommonUI/Public/Input/UIActionBinding.h
--- a/Engine/Plugins/Experimental/CommonUI/Source/CommonUI/Public/Input/UIActionBinding.h
+++ b/Engine/Plugins/Experimental/CommonUI/Source/CommonUI/Public/Input/UIActionBinding.h
@@ -13,6 +13,7 @@ struct FUIActionKeyMapping;
 class FActionRouterBindingCollection;
 struct FBindUIActionArgs;
 struct FCommonInputActionDataBase;
+class UInputAction;
 
 enum class EProcessHoldActionResult
 {
@@ -34,7 +35,7 @@ struct COMMONUI_API FUIActionBinding
 
 	bool operator==(const FUIActionBindingHandle& OtherHandle) const { return Handle == OtherHandle; }
 
-	// @TODO: DarenC - Remove legacy.
+	// @TODO: Rename non-legacy in 5.3. We no longer have any active plans to remove data tables in CommonUI.
 	FCommonInputActionDataBase* GetLegacyInputActionData() const;
 
 	EProcessHoldActionResult ProcessHoldInput(ECommonInputMode ActiveInputMode, FKey Key, EInputEvent InputEvent);
@@ -68,9 +69,11 @@ struct COMMONUI_API FUIActionBinding
 	DECLARE_MULTICAST_DELEGATE_OneParam(FOnHoldActionProgressedMulticast, float);
 	FOnHoldActionProgressedMulticast OnHoldActionProgressed;
 
-	// @TODO: DarenC - Remove legacy.
+	// @TODO: Rename non-legacy in 5.3. We no longer have any active plans to remove data tables in CommonUI.
 	FDataTableRowHandle LegacyActionTableRow;
 
+	TWeakObjectPtr<const UInputAction> InputAction;
+
 private:
 	FUIActionBinding(const UWidget& InBoundWidget, const FBindUIActionArgs& BindArgs);
 	
