Additional non-unity fix
#rnx

diff --git a/Engine/Plugins/Experimental/CommonUI/Source/CommonUI/Public/CommonUITypes.h b/Engine/Plugins/Experimental/CommonUI/Source/CommonUI/Public/CommonUITypes.h
--- a/Engine/Plugins/Experimental/CommonUI/Source/CommonUI/Public/CommonUITypes.h
+++ b/Engine/Plugins/Experimental/CommonUI/Source/CommonUI/Public/CommonUITypes.h
@@ -6,6 +6,7 @@
 #include "InputCoreTypes.h"
 #include "Styling/SlateBrush.h"
 #include "UObject/Interface.h"
+#include "UObject/ObjectMacros.h"
 #include "CommonUITypes.generated.h"
 
 class UUserWidget;
