Add the Event click to CommonButton.

diff --git a/Engine/Plugins/Experimental/CommonUI/Source/CommonUI/Private/CommonButtonBase.cpp b/Engine/Plugins/Experimental/CommonUI/Source/CommonUI/Private/CommonButtonBase.cpp
--- a/Engine/Plugins/Experimental/CommonUI/Source/CommonUI/Private/CommonButtonBase.cpp
+++ b/Engine/Plugins/Experimental/CommonUI/Source/CommonUI/Private/CommonButtonBase.cpp
@@ -1271,6 +1271,7 @@ void UCommonButtonBase::NativeOnClicked()
 	{
 		BP_OnClicked();
 		OnClicked().Broadcast();
+		BroadcastFieldValueChanged(FFieldNotificationClassDescriptor::ClickEvent);
 		if (OnButtonBaseClicked.IsBound())
 		{
 			OnButtonBaseClicked.Broadcast(this);
diff --git a/Engine/Plugins/Experimental/CommonUI/Source/CommonUI/Public/CommonButtonBase.h b/Engine/Plugins/Experimental/CommonUI/Source/CommonUI/Public/CommonButtonBase.h
--- a/Engine/Plugins/Experimental/CommonUI/Source/CommonUI/Public/CommonButtonBase.h
+++ b/Engine/Plugins/Experimental/CommonUI/Source/CommonUI/Public/CommonButtonBase.h
@@ -5,6 +5,7 @@
 #include "CommonUserWidget.h"
 #include "Components/Button.h"
 #include "Engine/DataTable.h"
+#include "FieldNotification/WidgetEventField.h"
 #include "Types/ISlateMetaData.h"
 #include "CommonInputModeTypes.h"
 #include "CommonButtonBase.generated.h"
@@ -473,6 +474,9 @@ public:
 	FCommonButtonEvent& OnLockClicked() const { return OnLockClickedEvent; }
 	FCommonButtonEvent& OnLockDoubleClicked() const { return OnLockDoubleClickedEvent; }
 
+	UPROPERTY(BlueprintReadOnly, FieldNotify, Category= "Common Button")
+	FWidgetEventField ClickEvent;
+
 	DECLARE_EVENT_OneParam(UCommonButtonBase, FOnIsSelectedChanged, bool);
 	FOnIsSelectedChanged& OnIsSelectedChanged() const { return OnIsSelectedChangedEvent; }
 
diff --git a/Engine/Source/Runtime/UMG/Public/FieldNotification/WidgetEventField.h b/Engine/Source/Runtime/UMG/Public/FieldNotification/WidgetEventField.h
new file mode 100644
--- /dev/null
+++ b/Engine/Source/Runtime/UMG/Public/FieldNotification/WidgetEventField.h
@@ -0,0 +1,26 @@
+// Copyright Epic Games, Inc. All Rights Reserved.
+
+#pragma once
+
+#include "WidgetEventField.generated.h"
+
+/**
+ * Generic structure to notify when an event occurs.
+ * 
+ * class UMyWidget : public UWidget
+ * {
+ *   UPROPERTY(FieldNotify)
+ *   FWidgetEventField SomeEvent;
+ *
+ *   void OnSomeEvent()
+ *   {
+ *     BroadcastFieldValueChanged(FFieldNotificationClassDescriptor::Text);
+ *   }
+ * };
+};
+ */
+USTRUCT(BlueprintType)
+struct FWidgetEventField
+{
+	GENERATED_BODY()
+};
