Minor CommonUI EnhancedInput fixes
Get Icon for first active input method key instead of first input action key.
Show enhanced UI input actions in bound action bar.

diff --git a/Engine/Plugins/Experimental/CommonUI/Source/CommonInput/Private/CommonInputSubsystem.cpp b/Engine/Plugins/Experimental/CommonUI/Source/CommonInput/Private/CommonInputSubsystem.cpp
--- a/Engine/Plugins/Experimental/CommonUI/Source/CommonInput/Private/CommonInputSubsystem.cpp
+++ b/Engine/Plugins/Experimental/CommonUI/Source/CommonInput/Private/CommonInputSubsystem.cpp
@@ -9,6 +9,7 @@
 #include "Engine/LocalPlayer.h"
 #include "Engine/PlatformSettingsManager.h"
 #include "Engine/World.h"
+#include "EnhancedInputSubsystems.h"
 #include "Framework/Application/IInputProcessor.h"
 #include "Framework/Application/SlateApplication.h"
 #include "Framework/Application/SlateUser.h"
@@ -354,6 +355,18 @@ void UCommonInputSubsystem::Initialize(FSubsystemCollectionBase& Collection)
 
 	CVarInputKeysVisible->SetOnChangedCallback(FConsoleVariableDelegate::CreateUObject(this, &UCommonInputSubsystem::ShouldShowInputKeysChanged));
 
+	if (ICommonInputModule::Get().GetSettings().GetEnableEnhancedInputSupport())
+	{
+		if (ULocalPlayer* LocalPlayer = GetLocalPlayerChecked())
+		{
+			if (UEnhancedInputLocalPlayerSubsystem* EnhancedInputLocalPlayerSubsystem = LocalPlayer->GetSubsystem<UEnhancedInputLocalPlayerSubsystem>())
+			{
+				BroadcastInputMethodChangedEvent.BindUFunction(this, GET_FUNCTION_NAME_CHECKED(UCommonInputSubsystem, BroadcastInputMethodChanged));
+				EnhancedInputLocalPlayerSubsystem->ControlMappingsRebuiltDelegate.AddUnique(BroadcastInputMethodChangedEvent);
+			}
+		}
+	}
+
 	SetActionDomainTable(FCommonInputBase::GetInputSettings()->GetActionDomainTable());
 }
 
@@ -366,6 +379,18 @@ void UCommonInputSubsystem::Deinitialize()
 	}
 	CommonInputPreprocessor.Reset();
 
+	if (ICommonInputModule::Get().GetSettings().GetEnableEnhancedInputSupport())
+	{
+		if (ULocalPlayer* LocalPlayer = GetLocalPlayerChecked())
+		{
+			if (UEnhancedInputLocalPlayerSubsystem* EnhancedInputLocalPlayerSubsystem = LocalPlayer->GetSubsystem<UEnhancedInputLocalPlayerSubsystem>())
+			{
+				EnhancedInputLocalPlayerSubsystem->ControlMappingsRebuiltDelegate.Remove(BroadcastInputMethodChangedEvent);
+				BroadcastInputMethodChangedEvent.Unbind();
+			}
+		}
+	}
+
 	FTSTicker::GetCoreTicker().RemoveTicker(TickHandle);
 }
 
diff --git a/Engine/Plugins/Experimental/CommonUI/Source/CommonInput/Public/CommonInputSubsystem.h b/Engine/Plugins/Experimental/CommonUI/Source/CommonInput/Public/CommonInputSubsystem.h
--- a/Engine/Plugins/Experimental/CommonUI/Source/CommonInput/Public/CommonInputSubsystem.h
+++ b/Engine/Plugins/Experimental/CommonUI/Source/CommonInput/Public/CommonInputSubsystem.h
@@ -108,6 +108,7 @@ public:
 protected:
 	ECommonInputType LockInput(ECommonInputType InputToLock) const;
 
+	UFUNCTION()
 	void BroadcastInputMethodChanged();
 
 private:
@@ -172,4 +173,7 @@ private:
 	* Note : Calling order is not guaranteed. Also, keep in mind that you might need to honor the previous callee's request to not support the input type being tested.
 	*/
 	static FPlatformInputSupportOverrideDelegate OnPlatformInputSupportOverride;
+
+	/** Used to broadcast when input method should be considered changed based on external systems, such as on enhanced input mapping context application*/
+	FScriptDelegate BroadcastInputMethodChangedEvent;
 };
\ No newline at end of file
diff --git a/Engine/Plugins/Experimental/CommonUI/Source/CommonUI/Private/CommonActionWidget.cpp b/Engine/Plugins/Experimental/CommonUI/Source/CommonUI/Private/CommonActionWidget.cpp
--- a/Engine/Plugins/Experimental/CommonUI/Source/CommonUI/Private/CommonActionWidget.cpp
+++ b/Engine/Plugins/Experimental/CommonUI/Source/CommonUI/Private/CommonActionWidget.cpp
@@ -210,6 +210,14 @@ void UCommonActionWidget::SetInputActionBinding(FUIActionBindingHandle BindingHa
 	UpdateBindingHandleInternal(BindingHandle);
 	if (TSharedPtr<FUIActionBinding> Binding = FUIActionBinding::FindBinding(BindingHandle))
 	{
+		if (CommonUI::IsEnhancedInputSupportEnabled())
+		{
+			if (const UInputAction* InputAction = Binding->InputAction.Get())
+			{
+				EnhancedInputAction = const_cast<UInputAction*>(InputAction);
+			}
+		}
+
 		InputActions.Reset();
 		InputActions.Add(Binding->LegacyActionTableRow);
 
diff --git a/Engine/Plugins/Experimental/CommonUI/Source/CommonUI/Private/CommonButtonBase.cpp b/Engine/Plugins/Experimental/CommonUI/Source/CommonUI/Private/CommonButtonBase.cpp
--- a/Engine/Plugins/Experimental/CommonUI/Source/CommonUI/Private/CommonButtonBase.cpp
+++ b/Engine/Plugins/Experimental/CommonUI/Source/CommonUI/Private/CommonButtonBase.cpp
@@ -1071,8 +1071,10 @@ void UCommonButtonBase::UpdateInputActionWidget()
 	// Update the input action state of the input action widget contextually based on the current state of the button
 	if (GetGameInstance() && InputActionWidget)
 	{
+		bool bIsEnhancedInputSupportEnabled = CommonUI::IsEnhancedInputSupportEnabled();
+
 		// Prefer visualizing the triggering enhanced input action before all else
-		if (CommonUI::IsEnhancedInputSupportEnabled() && TriggeringEnhancedInputAction)
+		if (bIsEnhancedInputSupportEnabled && TriggeringEnhancedInputAction)
 		{
 			InputActionWidget->SetEnhancedInputAction(TriggeringEnhancedInputAction);
 		}
@@ -1091,7 +1093,6 @@ void UCommonButtonBase::UpdateInputActionWidget()
 		{
 			FDataTableRowHandle HoverStateHandle;
 			UInputAction* HoverEnhancedInputAction = nullptr;
-			bool bIsEnhancedInputSupportEnabled = CommonUI::IsEnhancedInputSupportEnabled();
 			if (IsHovered())
 			{
 				if (bIsEnhancedInputSupportEnabled)
@@ -1115,6 +1116,11 @@ void UCommonButtonBase::UpdateInputActionWidget()
 		}
 		else
 		{
+			if (bIsEnhancedInputSupportEnabled)
+			{
+				InputActionWidget->SetEnhancedInputAction(nullptr);
+			}
+
 			FDataTableRowHandle EmptyStateHandle;
 			InputActionWidget->SetInputAction(EmptyStateHandle);
 		}
diff --git a/Engine/Plugins/Experimental/CommonUI/Source/CommonUI/Private/CommonUITypes.cpp b/Engine/Plugins/Experimental/CommonUI/Source/CommonUI/Private/CommonUITypes.cpp
--- a/Engine/Plugins/Experimental/CommonUI/Source/CommonUI/Private/CommonUITypes.cpp
+++ b/Engine/Plugins/Experimental/CommonUI/Source/CommonUI/Private/CommonUITypes.cpp
@@ -249,19 +249,16 @@ bool CommonUI::IsEnhancedInputSupportEnabled()
 	return bEnabled;
 }
 
-TObjectPtr<const UCommonInputMetadata> CommonUI::GetEnhancedInputActionMetadata(const ULocalPlayer* LocalPlayer, const UInputAction* InputAction)
+TObjectPtr<const UCommonInputMetadata> CommonUI::GetEnhancedInputActionMetadata(const UInputAction* InputAction)
 {
 	if (!InputAction->GetPlayerMappableKeySettings())
 	{
 		return nullptr;
 	}
 
-	if (UBlueprint* MetadataBP = Cast<UBlueprint>(InputAction->GetPlayerMappableKeySettings()->Metadata))
+	if (ICommonMappingContextMetadataInterface* MappingContextMetadata = Cast<ICommonMappingContextMetadataInterface>(InputAction->GetPlayerMappableKeySettings()->Metadata))
 	{
-		if (ICommonMappingContextMetadataInterface* MappingContextMetadata = Cast<ICommonMappingContextMetadataInterface>(MetadataBP->GeneratedClass->GetDefaultObject()))
-		{
-			return MappingContextMetadata->GetCommonInputMetadata(InputAction);
-		}
+		return MappingContextMetadata->GetCommonInputMetadata(InputAction);
 	}
 
 	return nullptr;
@@ -291,20 +288,83 @@ void CommonUI::InjectEnhancedInputForAction(const ULocalPlayer* LocalPlayer, con
 
 FSlateBrush CommonUI::GetIconForEnhancedInputAction(const UCommonInputSubsystem* CommonInputSubsystem, const UInputAction* InputAction)
 {
+	FKey FirstKeyForCurrentInput = GetFirstKeyForInputType(CommonInputSubsystem->GetLocalPlayer(), CommonInputSubsystem->GetCurrentInputType(), InputAction);
+
+	FSlateBrush SlateBrush;
+	if (FirstKeyForCurrentInput.IsValid() && UCommonInputPlatformSettings::Get()->TryGetInputBrush(SlateBrush, FirstKeyForCurrentInput, CommonInputSubsystem->GetCurrentInputType(), CommonInputSubsystem->GetCurrentGamepadName()))
+	{
+		return SlateBrush;
+	}
+
+	return *FStyleDefaults::GetNoBrush();
+}
+
+bool CommonUI::ActionValidForInputType(const ULocalPlayer* LocalPlayer, ECommonInputType InputType, const UInputAction* InputAction)
+{
+	if (!LocalPlayer)
+	{
+		return false;
+	}
+
 	TArray<FKey> Keys;
-	if (ULocalPlayer* LocalPlayer = CommonInputSubsystem->GetLocalPlayer())
+	if (UEnhancedInputLocalPlayerSubsystem* EnhancedInputLocalPlayerSubsystem = LocalPlayer->GetSubsystem<UEnhancedInputLocalPlayerSubsystem>())
 	{
-		if (UEnhancedInputLocalPlayerSubsystem* EnhancedInputLocalPlayerSubsystem = LocalPlayer->GetSubsystem<UEnhancedInputLocalPlayerSubsystem>())
+		Keys = EnhancedInputLocalPlayerSubsystem->QueryKeysMappedToAction(InputAction);
+	}
+
+	for (const FKey& Key : Keys)
+	{
+		if (!Key.IsValid())
+		{
+			continue;
+		}
+
+		bool bIsValidTouch = Key.IsTouch() && InputType == ECommonInputType::Touch;
+		bool bIsValidGamepad = Key.IsGamepadKey() && InputType == ECommonInputType::Gamepad;
+		bool bIsValidMouseAndKeyboard = !Key.IsTouch() && !Key.IsGamepadKey() && InputType == ECommonInputType::MouseAndKeyboard;
+
+		if (bIsValidTouch || bIsValidGamepad || bIsValidMouseAndKeyboard)
 		{
-			Keys = EnhancedInputLocalPlayerSubsystem->QueryKeysMappedToAction(InputAction);
+			return true;
 		}
 	}
 
-	FSlateBrush SlateBrush;
-	if (Keys.Num() && UCommonInputPlatformSettings::Get()->TryGetInputBrush(SlateBrush, Keys[0], CommonInputSubsystem->GetCurrentInputType(), CommonInputSubsystem->GetCurrentGamepadName()))
+	return false;
+}
+
+FKey CommonUI::GetFirstKeyForInputType(const ULocalPlayer* LocalPlayer, ECommonInputType InputType, const UInputAction* InputAction)
+{
+	if (!LocalPlayer)
 	{
-		return SlateBrush;
+		return FKey();
 	}
 
-	return *FStyleDefaults::GetNoBrush();
+	TArray<FKey> Keys;
+	if (UEnhancedInputLocalPlayerSubsystem* EnhancedInputLocalPlayerSubsystem = LocalPlayer->GetSubsystem<UEnhancedInputLocalPlayerSubsystem>())
+	{
+		Keys = EnhancedInputLocalPlayerSubsystem->QueryKeysMappedToAction(InputAction);
+	}
+
+	for (const FKey& Key : Keys)
+	{
+		if (!Key.IsValid())
+		{
+			continue;
+		}
+
+		if (Key.IsTouch() && InputType == ECommonInputType::Touch)
+		{
+			return Key;
+		}
+		else if (Key.IsGamepadKey() && InputType == ECommonInputType::Gamepad)
+		{
+			return Key;
+		}
+		else if (!Key.IsTouch() && !Key.IsGamepadKey() && InputType == ECommonInputType::MouseAndKeyboard)
+		{
+			return Key;
+		}
+	}
+
+	return FKey();
 }
diff --git a/Engine/Plugins/Experimental/CommonUI/Source/CommonUI/Private/Input/CommonBoundActionBar.cpp b/Engine/Plugins/Experimental/CommonUI/Source/CommonUI/Private/Input/CommonBoundActionBar.cpp
--- a/Engine/Plugins/Experimental/CommonUI/Source/CommonUI/Private/Input/CommonBoundActionBar.cpp
+++ b/Engine/Plugins/Experimental/CommonUI/Source/CommonUI/Private/Input/CommonBoundActionBar.cpp
@@ -10,6 +10,7 @@
 #include "Input/CommonBoundActionButtonInterface.h"
 #include "Input/CommonUIActionRouterBase.h"
 #include "Input/UIActionBinding.h"
+#include "InputAction.h"
 #include "TimerManager.h"
 
 #include UE_INLINE_GENERATED_CPP_BY_NAME(CommonBoundActionBar)
@@ -161,7 +162,7 @@ void UCommonBoundActionBar::HandleDeferredDisplayUpdate()
 					const FName& PlayerGamepadName = InputSubsystem.GetCurrentGamepadName();
 
 					TSet<FName> AcceptedBindings;
-					TArray<FUIActionBindingHandle> FilteredBindings = ActionRouter->GatherActiveBindings().FilterByPredicate([PlayerInputType, PlayerGamepadName, &AcceptedBindings](const FUIActionBindingHandle& Handle) mutable
+					TArray<FUIActionBindingHandle> FilteredBindings = ActionRouter->GatherActiveBindings().FilterByPredicate([ActionRouter, PlayerInputType, PlayerGamepadName, &AcceptedBindings](const FUIActionBindingHandle& Handle) mutable
 						{
 							if (TSharedPtr<FUIActionBinding> Binding = FUIActionBinding::FindBinding(Handle))
 							{
@@ -170,6 +171,14 @@ void UCommonBoundActionBar::HandleDeferredDisplayUpdate()
 									return false;
 								}
 
+								if (CommonUI::IsEnhancedInputSupportEnabled())
+								{
+									if (TObjectPtr<const UInputAction> InputAction = Binding->InputAction.Get())
+									{
+										return CommonUI::ActionValidForInputType(ActionRouter->GetLocalPlayer(), PlayerInputType, InputAction);
+									}
+								}
+
 								if (FCommonInputActionDataBase* LegacyData = Binding->GetLegacyInputActionData())
 								{
 									if (!LegacyData->CanDisplayInReflector(PlayerInputType, PlayerGamepadName))
@@ -192,45 +201,98 @@ void UCommonBoundActionBar::HandleDeferredDisplayUpdate()
 
 					//Force Virtual_Back to one end of the list so Back actions are always consistent.
 					//Otherwise, order within a node is controlled by order of add/remove.
-					Algo::Sort(FilteredBindings, [PlayerInputType, PlayerGamepadName](const FUIActionBindingHandle& A, const FUIActionBindingHandle& B)
+					Algo::Sort(FilteredBindings, [ActionRouter, PlayerInputType, PlayerGamepadName](const FUIActionBindingHandle& A, const FUIActionBindingHandle& B)
 					{
 						TSharedPtr<FUIActionBinding> BindingA = FUIActionBinding::FindBinding(A);
 						TSharedPtr<FUIActionBinding> BindingB = FUIActionBinding::FindBinding(B);
 
 						if (ensureMsgf((BindingA && BindingB), TEXT("The array filter above should enforce that there are no null bindings")))
 						{
-							FCommonInputActionDataBase* LegacyDataA = BindingA->GetLegacyInputActionData();
-							FCommonInputActionDataBase* LegacyDataB = BindingB->GetLegacyInputActionData();
-							
-							if (ensureMsgf((LegacyDataA && LegacyDataB), TEXT("Non-legacy bindings not displayed yet -- array filter enforces they are not included")))
+							auto IsKeyBackAction = [ActionRouter, PlayerInputType, PlayerGamepadName](FCommonInputActionDataBase* LegacyData, const UInputAction* InputAction)
 							{
-								FKey KeyA = LegacyDataA->GetInputTypeInfo(PlayerInputType, PlayerGamepadName).GetKey();
-								FKey KeyB = LegacyDataB->GetInputTypeInfo(PlayerInputType, PlayerGamepadName).GetKey();
-							
-								// Fallback back to keyboard key when there is no key for touch
-								if (PlayerInputType == ECommonInputType::Touch)
+								if (LegacyData)
 								{
-									if (!KeyA.IsValid())
+									FKey Key = LegacyData->GetInputTypeInfo(PlayerInputType, PlayerGamepadName).GetKey();
+
+									// Fallback back to keyboard key when there is no key for touch
+									if (PlayerInputType == ECommonInputType::Touch)
 									{
-										KeyA = LegacyDataA->GetInputTypeInfo(ECommonInputType::MouseAndKeyboard, PlayerGamepadName).GetKey();
+										if (!Key.IsValid())
+										{
+											Key = LegacyData->GetInputTypeInfo(ECommonInputType::MouseAndKeyboard, PlayerGamepadName).GetKey();
+										}
 									}
 
-									if (!KeyB.IsValid())
+									return Key == EKeys::Virtual_Back || Key == EKeys::Escape || Key == EKeys::Android_Back;
+								}
+								else if (InputAction)
+								{
+									FKey Key = CommonUI::GetFirstKeyForInputType(ActionRouter->GetLocalPlayer(), PlayerInputType, InputAction);
+
+									// Fallback back to keyboard key when there is no key for touch
+									if (PlayerInputType == ECommonInputType::Touch)
+									{
+										if (!Key.IsValid())
+										{
+											Key = CommonUI::GetFirstKeyForInputType(ActionRouter->GetLocalPlayer(), ECommonInputType::MouseAndKeyboard, InputAction);
+										}
+									}
+
+									return Key == EKeys::Virtual_Back || Key == EKeys::Escape || Key == EKeys::Android_Back;
+								}
+
+								return false;
+							};
+
+							auto GetNavBarPriority = [](FCommonInputActionDataBase* LegacyData, const UInputAction* InputAction)
+							{
+								if (LegacyData)
+								{
+									return LegacyData->NavBarPriority;
+								}
+								else if (InputAction)
+								{
+									if (TObjectPtr<const UCommonInputMetadata> InputActionMetadata = CommonUI::GetEnhancedInputActionMetadata(InputAction))
 									{
-										KeyB = LegacyDataB->GetInputTypeInfo(ECommonInputType::MouseAndKeyboard, PlayerGamepadName).GetKey();
+										return InputActionMetadata->NavBarPriority;
 									}
 								}
 
-								bool bAIsBack = (KeyA == EKeys::Virtual_Back || KeyA == EKeys::Escape || KeyA == EKeys::Android_Back);
-								bool bBIsBack = (KeyB == EKeys::Virtual_Back || KeyB == EKeys::Escape || KeyB == EKeys::Android_Back);
-							
+								return 0;
+							};
+
+							FCommonInputActionDataBase* LegacyDataA = BindingA->GetLegacyInputActionData();
+							FCommonInputActionDataBase* LegacyDataB = BindingB->GetLegacyInputActionData();
+
+							const UInputAction* InputActionA = nullptr;
+							const UInputAction* InputActionB = nullptr;
+
+							// InputActions will be null unless IsEnhancedInputSupportEnabled, so we don't have to check for support below
+							if (CommonUI::IsEnhancedInputSupportEnabled())
+							{
+								InputActionA = BindingA->InputAction.Get();
+								InputActionB = BindingB->InputAction.Get();
+							}
+
+							bool bIsValidActionA = LegacyDataA || InputActionA;
+							bool bIsValidActionB = LegacyDataB || InputActionB;
+
+							if (ensureMsgf((bIsValidActionA && bIsValidActionB), TEXT("Action bindings not displayed yet -- array filter enforces they are not included")))
+							{
+								bool bAIsBack = IsKeyBackAction(LegacyDataA, InputActionA);
+								bool bBIsBack = IsKeyBackAction(LegacyDataB, InputActionB);
+
 								if (bAIsBack && !bBIsBack)
 								{
 									return false;
 								}
-								else if (LegacyDataA->NavBarPriority != LegacyDataB->NavBarPriority)
+
+								int32 NavBarPriorityA = GetNavBarPriority(LegacyDataA, InputActionA);
+								int32 NavBarPriorityB = GetNavBarPriority(LegacyDataB, InputActionB);
+
+								if (NavBarPriorityA != NavBarPriorityB)
 								{
-									return LegacyDataA->NavBarPriority < LegacyDataB->NavBarPriority;
+									return NavBarPriorityA < NavBarPriorityB;
 								}
 							}
 
diff --git a/Engine/Plugins/Experimental/CommonUI/Source/CommonUI/Private/Input/UIActionRouterTypes.cpp b/Engine/Plugins/Experimental/CommonUI/Source/CommonUI/Private/Input/UIActionRouterTypes.cpp
--- a/Engine/Plugins/Experimental/CommonUI/Source/CommonUI/Private/Input/UIActionRouterTypes.cpp
+++ b/Engine/Plugins/Experimental/CommonUI/Source/CommonUI/Private/Input/UIActionRouterTypes.cpp
@@ -403,6 +403,14 @@ FText FUIActionBindingHandle::GetDisplayName() const
 	{
 		if (const UCommonInputSubsystem* CommonInputSubsystem = Binding->BoundWidget.IsValid() ? UCommonInputSubsystem::Get(Binding->BoundWidget->GetOwningLocalPlayer()) : nullptr)
 		{
+			if (CommonUI::IsEnhancedInputSupportEnabled())
+			{
+				if (const TObjectPtr<const UInputAction> InputAction = Binding->InputAction.Get())
+				{
+					return InputAction->ActionDescription;
+				}
+			}
+
 			const ECommonInputType CurrentInputType = CommonInputSubsystem->GetCurrentInputType();
 
 			for (const FUIActionKeyMapping& HoldMapping : Binding->HoldMappings)
@@ -551,7 +559,7 @@ bool FActionRouterBindingCollection::ProcessNormalInput(ECommonInputMode ActiveI
 						bool bEnhancedInputInjected = false;
 						if (InInputAction)
 						{
-							if (TObjectPtr<const UCommonInputMetadata> CommonInputMetadata = CommonUI::GetEnhancedInputActionMetadata(GetActionRouter().GetLocalPlayerChecked(), InInputAction))
+							if (TObjectPtr<const UCommonInputMetadata> CommonInputMetadata = CommonUI::GetEnhancedInputActionMetadata(InInputAction))
 							{
 								// Non generic actions should inject enhanced input so users can bind to enhanced input events
 								if (!CommonInputMetadata->bIsGenericInputAction)
diff --git a/Engine/Plugins/Experimental/CommonUI/Source/CommonUI/Public/CommonUITypes.h b/Engine/Plugins/Experimental/CommonUI/Source/CommonUI/Public/CommonUITypes.h
--- a/Engine/Plugins/Experimental/CommonUI/Source/CommonUI/Public/CommonUITypes.h
+++ b/Engine/Plugins/Experimental/CommonUI/Source/CommonUI/Public/CommonUITypes.h
@@ -3,6 +3,7 @@
 #pragma once
 
 #include "Engine/DataTable.h"
+#include "Engine/DataAsset.h"
 #include "InputCoreTypes.h"
 #include "Styling/SlateBrush.h"
 #include "UObject/Interface.h"
@@ -245,8 +246,8 @@ public:
  * Utilizes a map of input actions to metadata to prevent users from having to create
  * multiple metadata assets / instances. Using this map is not mandatory.
  */
-UCLASS(Blueprintable)
-class COMMONUI_API UCommonMappingContextMetadata : public UObject, public ICommonMappingContextMetadataInterface
+UCLASS(BlueprintType)
+class COMMONUI_API UCommonMappingContextMetadata : public UDataAsset, public ICommonMappingContextMetadataInterface
 {
 	GENERATED_BODY()
 
@@ -274,10 +275,12 @@ public:
 
 	static bool IsEnhancedInputSupportEnabled();
 
-	static TObjectPtr<const UCommonInputMetadata> GetEnhancedInputActionMetadata(const ULocalPlayer* LocalPlayer, const UInputAction* InputAction);
+	static TObjectPtr<const UCommonInputMetadata> GetEnhancedInputActionMetadata(const UInputAction* InputAction);
 	static void GetEnhancedInputActionKeys(const ULocalPlayer* LocalPlayer, const UInputAction* InputAction, TArray<FKey>& OutKeys);
 	static void InjectEnhancedInputForAction(const ULocalPlayer* LocalPlayer, const UInputAction* InputAction, FInputActionValue RawValue);
 	static FSlateBrush GetIconForEnhancedInputAction(const UCommonInputSubsystem* CommonInputSubsystem, const UInputAction* InputAction);
+	static bool ActionValidForInputType(const ULocalPlayer* LocalPlayer, ECommonInputType InputType, const UInputAction* InputAction);
+	static FKey GetFirstKeyForInputType(const ULocalPlayer* LocalPlayer, ECommonInputType InputType, const UInputAction* InputAction);
 };
 
 DECLARE_DYNAMIC_DELEGATE_OneParam(FOnItemClicked, UUserWidget*, Widget);
