PR #10126: CommonUI Fix - SCommonTileView and SCommonListView to use FName lists. (Contributed by [REDACTED_PERSON])

diff --git a/Engine/Plugins/Experimental/CommonUI/Source/CommonUI/Public/CommonListView.h b/Engine/Plugins/Experimental/CommonUI/Source/CommonUI/Public/CommonListView.h
--- a/Engine/Plugins/Experimental/CommonUI/Source/CommonUI/Public/CommonListView.h
+++ b/Engine/Plugins/Experimental/CommonUI/Source/CommonUI/Public/CommonListView.h
@@ -23,7 +23,7 @@ public:
 			// If bReturnFocusToSelection is true find the last selected object and focus on that.
 			if (this->GetItems().Num() > 0)
 			{
-				typename TListTypeTraits<ItemType>::NullableType ItemNavigatedTo(nullptr);
+				typename TListTypeTraits<ItemType>::NullableType ItemNavigatedTo = TListTypeTraits<ItemType>::MakeNullPtr();
 				if (this->GetNumItemsSelected() == 0)
 				{
 					ItemNavigatedTo = this->GetItems()[0];
diff --git a/Engine/Plugins/Experimental/CommonUI/Source/CommonUI/Public/CommonTreeView.h b/Engine/Plugins/Experimental/CommonUI/Source/CommonUI/Public/CommonTreeView.h
--- a/Engine/Plugins/Experimental/CommonUI/Source/CommonUI/Public/CommonTreeView.h
+++ b/Engine/Plugins/Experimental/CommonUI/Source/CommonUI/Public/CommonTreeView.h
@@ -23,7 +23,7 @@ public:
 			// If bReturnFocusToSelection is true find the last selected object and focus on that.
 			if (this->HasValidItemsSource() && this->TreeItemsSource->Num() > 0)
 			{
-				typename TListTypeTraits<ItemType>::NullableType ItemNavigatedTo(nullptr);
+				typename TListTypeTraits<ItemType>::NullableType ItemNavigatedTo = TListTypeTraits<ItemType>::MakeNullPtr();
 				if (this->GetNumItemsSelected() == 0)
 				{
 					ItemNavigatedTo = (*this->TreeItemsSource)[0];
