Adding COMMONUI_API to SAnalogSlider, it's in a public header but not actually exported.

diff --git a/Engine/Plugins/Experimental/CommonUI/Source/CommonUI/Public/SAnalogSlider.h b/Engine/Plugins/Experimental/CommonUI/Source/CommonUI/Public/SAnalogSlider.h
--- a/Engine/Plugins/Experimental/CommonUI/Source/CommonUI/Public/SAnalogSlider.h
+++ b/Engine/Plugins/Experimental/CommonUI/Source/CommonUI/Public/SAnalogSlider.h
@@ -9,7 +9,7 @@ struct FAnalogInputEvent;
 /**
  * A Slate slider control is a linear scale and draggable handle.
  */
-class SAnalogSlider : public SSlider
+class COMMONUI_API SAnalogSlider : public SSlider
 {
 public:
 
