Fixed an edge case where we were refreshing action domain leaf node config when the activatable tree was disabled.
#rnx
#tests: Ran a few console tests similar to the QA tests with SplitScreen and game modes.

#localization none

diff --git a/Engine/Plugins/Experimental/CommonUI/Source/CommonUI/Private/Input/CommonUIActionRouterBase.cpp b/Engine/Plugins/Experimental/CommonUI/Source/CommonUI/Private/Input/CommonUIActionRouterBase.cpp
--- a/Engine/Plugins/Experimental/CommonUI/Source/CommonUI/Private/Input/CommonUIActionRouterBase.cpp
+++ b/Engine/Plugins/Experimental/CommonUI/Source/CommonUI/Private/Input/CommonUIActionRouterBase.cpp
@@ -680,7 +680,7 @@ void UCommonUIActionRouterBase::HandleRootNodeDeactivated(TWeakPtr<FActivatableT
 	}
 
 	// In the case that all root nodes are not activated we need to re-establish input for the highest paint layer node in action domain nodes.
-	if (!bActivatedRootNodeExists)
+	if (!bActivatedRootNodeExists && bIsActivatableTreeEnabled)
 	{
 		RefreshActionDomainLeafNodeConfig();
 	}
