Removed style sheet prototype code from CommonTextBlock. Removed CommonStyleSheet class from CommonUI module.

diff --git a/Engine/Plugins/Experimental/CommonUI/Source/CommonUI/Private/CommonTextBlock.cpp b/Engine/Plugins/Experimental/CommonUI/Source/CommonUI/Private/CommonTextBlock.cpp
--- a/Engine/Plugins/Experimental/CommonUI/Source/CommonUI/Private/CommonTextBlock.cpp
+++ b/Engine/Plugins/Experimental/CommonUI/Source/CommonUI/Private/CommonTextBlock.cpp
@@ -8,7 +8,6 @@
 #include "Widgets/Text/STextBlock.h"
 #include "Layout/ArrangedChildren.h"
 #include "Layout/LayoutUtils.h"
-#include "StyleSheet/CommonStyleSheet.h"
 
 #include UE_INLINE_GENERATED_CPP_BY_NAME(CommonTextBlock)
 
@@ -650,11 +649,6 @@ void UCommonTextBlock::UpdateFromStyle()
 			SetShadowColorAndOpacity(FLinearColor::Transparent);
 		}
 	}
-
-	if (UCommonStyleSheet* StyleSheetPtr = StyleSheet.Get())
-	{
-		StyleSheetPtr->Apply(this);
-	}
 }
 
 void UCommonTextBlock::ReleaseSlateResources(bool bReleaseChildren)
diff --git a/Engine/Plugins/Experimental/CommonUI/Source/CommonUI/Private/StyleSheet/CommonStyleSheet.cpp b/Engine/Plugins/Experimental/CommonUI/Source/CommonUI/Private/StyleSheet/CommonStyleSheet.cpp
deleted file mode 100644
diff --git a/Engine/Plugins/Experimental/CommonUI/Source/CommonUI/Public/CommonTextBlock.h b/Engine/Plugins/Experimental/CommonUI/Source/CommonUI/Public/CommonTextBlock.h
--- a/Engine/Plugins/Experimental/CommonUI/Source/CommonUI/Public/CommonTextBlock.h
+++ b/Engine/Plugins/Experimental/CommonUI/Source/CommonUI/Public/CommonTextBlock.h
@@ -9,8 +9,6 @@
 
 struct FTextBlockStyle;
 
-class UCommonStyleSheet;
-
 /* 
  * ---- All properties must be EditDefaultsOnly, BlueprintReadOnly !!! -----
  * We return the CDO to blueprints, so we cannot allow any changes (blueprint doesn't support const variables)
@@ -233,10 +231,6 @@ private:
 	UPROPERTY(EditAnywhere, BlueprintReadOnly, Category = CommonText, meta = (ExposeOnSpawn = true, AllowPrivateAccess = true))
 	TSubclassOf<UCommonTextScrollStyle> ScrollStyle;
 
-	/** Prototype: Do not use! */
-	UPROPERTY(EditAnywhere, BlueprintReadOnly, AdvancedDisplay, Category = CommonText, meta = (ExposeOnSpawn = true, AllowPrivateAccess = true))
-	TObjectPtr<UCommonStyleSheet> StyleSheet;
-
 	/** If scrolling is enabled/disabled initially, this can be updated in blueprint */
 	UPROPERTY(EditAnywhere, BlueprintReadOnly, Category = CommonText, meta = (ExposeOnSpawn = true, AllowPrivateAccess = true))
 	bool bIsScrollingEnabled = true;
diff --git a/Engine/Plugins/Experimental/CommonUI/Source/CommonUI/Public/StyleSheet/CommonStyleSheet.h b/Engine/Plugins/Experimental/CommonUI/Source/CommonUI/Public/StyleSheet/CommonStyleSheet.h
deleted file mode 100644
diff --git a/Engine/Plugins/Experimental/CommonUI/Source/CommonUI/Public/StyleSheet/CommonStyleSheetTypes.h b/Engine/Plugins/Experimental/CommonUI/Source/CommonUI/Public/StyleSheet/CommonStyleSheetTypes.h
deleted file mode 100644
