Verify validity of CommonInputSubsystem before trying to get icon from it.
#rnx
#tests: Tested various paths in PIE that use these functions.

diff --git a/Engine/Plugins/Experimental/CommonUI/Source/CommonUI/Private/CommonActionWidget.cpp b/Engine/Plugins/Experimental/CommonUI/Source/CommonUI/Private/CommonActionWidget.cpp
--- a/Engine/Plugins/Experimental/CommonUI/Source/CommonUI/Private/CommonActionWidget.cpp
+++ b/Engine/Plugins/Experimental/CommonUI/Source/CommonUI/Private/CommonActionWidget.cpp
@@ -1,17 +1,19 @@
 // Copyright Epic Games, Inc. All Rights Reserved.
 
 #include "CommonActionWidget.h"
+
 #include "CommonInputSubsystem.h"
 #include "CommonInputTypeEnum.h"
-#include "CommonWidgetPaletteCategories.h"
-#include "Widgets/Layout/SBox.h"
-#include "Widgets/Images/SImage.h"
-#include "Widgets/SOverlay.h"
 #include "CommonUITypes.h"
-#include "Materials/MaterialInstanceDynamic.h"
+#include "CommonWidgetPaletteCategories.h"
+#include "Input/UIActionBinding.h"
 #include "InputAction.h"
 #include "InputTriggers.h"
-#include "Input/UIActionBinding.h"
+#include "Materials/MaterialInstanceDynamic.h"
+#include "Styling/StyleDefaults.h"
+#include "Widgets/Images/SImage.h"
+#include "Widgets/Layout/SBox.h"
+#include "Widgets/SOverlay.h"
 
 #include UE_INLINE_GENERATED_CPP_BY_NAME(CommonActionWidget)
 
@@ -110,9 +112,14 @@ void UCommonActionWidget::SynchronizeProperties()
 
 FSlateBrush UCommonActionWidget::GetIcon() const
 {
-	return EnhancedInputAction && CommonUI::IsEnhancedInputSupportEnabled()
-		? CommonUI::GetIconForEnhancedInputAction(GetInputSubsystem(), EnhancedInputAction)
-		: CommonUI::GetIconForInputActions(GetInputSubsystem(), InputActions);
+	if (const UCommonInputSubsystem* CommonInputSubsystem = GetInputSubsystem())
+	{
+		return EnhancedInputAction && CommonUI::IsEnhancedInputSupportEnabled()
+			? CommonUI::GetIconForEnhancedInputAction(CommonInputSubsystem, EnhancedInputAction)
+			: CommonUI::GetIconForInputActions(CommonInputSubsystem, InputActions);
+	}
+
+	return *FStyleDefaults::GetNoBrush();
 }
 
 UCommonInputSubsystem* UCommonActionWidget::GetInputSubsystem() const
