All windows platform extensions use the 'PC' input type by default.

#rnx

diff --git a/Engine/Plugins/Experimental/CommonUI/Source/CommonInput/CommonInput.Build.cs b/Engine/Plugins/Experimental/CommonUI/Source/CommonInput/CommonInput.Build.cs
--- a/Engine/Plugins/Experimental/CommonUI/Source/CommonInput/CommonInput.Build.cs
+++ b/Engine/Plugins/Experimental/CommonUI/Source/CommonInput/CommonInput.Build.cs
@@ -50,7 +50,7 @@ public class CommonInput : ModuleRules
 
 	static public string ToCommonUIPlatform(UnrealTargetPlatform TargetPlatform)
 	{
-		if (TargetPlatform == UnrealTargetPlatform.Win64)
+		if (TargetPlatform.IsInGroup(UnrealPlatformGroup.Windows))
 		{
 			return "PC";
 		}
