Hold Support for Common Button Base:
- Added flag to make the common button require hold to trigger action when focused and pressed
- Added optional Hold Behavior abstract class to specify Keyboard and Mouse, Gamepad and Touch hold behavior values.
- Modify Hold behavior for input bindings, including Rollback time for a required hold input binding.
- Added Default Hold Data in the Common UI Editor settings, under the Button category.
- Setting bRequiresHold to true will force hold on the bound Input Action of the button. UCommonButtonBase::GetConvertInputActionToHold() may be overridden to alter this behavior.

diff --git a/Engine/Plugins/Experimental/CommonUI/Source/CommonInput/Public/CommonInputBaseTypes.h b/Engine/Plugins/Experimental/CommonUI/Source/CommonInput/Public/CommonInputBaseTypes.h
--- a/Engine/Plugins/Experimental/CommonUI/Source/CommonInput/Public/CommonInputBaseTypes.h
+++ b/Engine/Plugins/Experimental/CommonUI/Source/CommonInput/Public/CommonInputBaseTypes.h
@@ -77,6 +77,19 @@ struct FInputDeviceIdentifierPair
 	FString HardwareDeviceIdentifier;
 };
 
+/* Data values needed for Hold interaction per input type. */
+USTRUCT()
+struct FInputHoldData
+{
+	GENERATED_BODY()
+
+	UPROPERTY(EditDefaultsOnly, Category = "Properties", meta = (ClampMin = "0.0", UIMin = "0.0"))
+	float HoldTime;
+
+	UPROPERTY(EditDefaultsOnly, Category = "Properties", meta = (ClampMin = "0.0", UIMin = "0.0"))
+	float HoldRollbackTime;
+};
+
 /* Derive from this class to store the Input data. It is referenced in the Common Input Settings, found in the project settings UI. */
 UCLASS(Abstract, Blueprintable, ClassGroup = Input, meta = (Category = "Common Input"))
 class COMMONINPUT_API UCommonUIInputData : public UObject
@@ -100,6 +113,33 @@ public:
 	TObjectPtr<UInputAction> EnhancedInputBackAction;
 };
 
+/* Defines values for hold behavior per input type: for mouse Press and Hold, for gamepad focused Press and Hold, for touch Press and Hold. */
+UCLASS(Abstract, Blueprintable, ClassGroup = Input, meta = (Category = "Common Input"))
+class COMMONINPUT_API UCommonUIHoldData : public UObject
+{
+	GENERATED_BODY()
+
+public:
+	UCommonUIHoldData()
+	{
+		KeyboardAndMouse.HoldTime = 0.75f;
+		KeyboardAndMouse.HoldRollbackTime = 0.0f;
+		Gamepad.HoldTime = 0.75f;
+		Gamepad.HoldRollbackTime = 0.0f;
+		Touch.HoldTime = 0.75f;
+		Touch.HoldRollbackTime = 0.0f;
+	}
+	
+	UPROPERTY(EditDefaultsOnly, Category = "KeyboardAndMouse", meta = (ShowOnlyInnerProperties))
+	FInputHoldData KeyboardAndMouse;
+
+	UPROPERTY(EditDefaultsOnly, Category = "Gamepad", meta = (ShowOnlyInnerProperties))
+	FInputHoldData Gamepad;
+
+	UPROPERTY(EditDefaultsOnly, Category = "Touch", meta = (ShowOnlyInnerProperties))
+	FInputHoldData Touch;
+};
+
 /* Derive from this class to store the Input data. It is referenced in the Common Input Settings, found in the project settings UI. */
 UCLASS(Abstract, Blueprintable, ClassGroup = Input, meta = (Category = "Common Input"))
 class COMMONINPUT_API UCommonInputBaseControllerData : public UObject
diff --git a/Engine/Plugins/Experimental/CommonUI/Source/CommonUI/Private/CommonButtonBase.cpp b/Engine/Plugins/Experimental/CommonUI/Source/CommonUI/Private/CommonButtonBase.cpp
--- a/Engine/Plugins/Experimental/CommonUI/Source/CommonUI/Private/CommonButtonBase.cpp
+++ b/Engine/Plugins/Experimental/CommonUI/Source/CommonUI/Private/CommonButtonBase.cpp
@@ -529,6 +529,26 @@ void UCommonButtonBase::BindTriggeringInputActionToClick()
 		return;
 	}
 
+	}
+}
+
+void UCommonButtonBase::BindTriggeringInputActionToClick()
+{
+	if (CommonUI::IsEnhancedInputSupportEnabled() && TriggeringEnhancedInputAction)
+	{
+		FBindUIActionArgs BindArgs(TriggeringEnhancedInputAction, false, FSimpleDelegate::CreateUObject(this, &UCommonButtonBase::HandleTriggeringActionCommited));
+		BindArgs.OnHoldActionProgressed.BindUObject(this, &UCommonButtonBase::NativeOnActionProgress);
+		BindArgs.OnHoldActionPressed.BindUObject(this, &UCommonButtonBase::HoldReset);
+        BindArgs.bIsPersistent = bIsPersistentBinding;
+        BindArgs.bForceHold = GetConvertInputActionToHold()
+
+		BindArgs.InputMode = InputModeOverride;
+
+		TriggeringBindingHandle = RegisterUIActionBinding(BindArgs);
+
+		return;
+	}
+
 	if (TriggeringInputAction.IsNull() || !TriggeredInputAction.IsNull())
 	{
 		return;
diff --git a/Engine/Plugins/Experimental/CommonUI/Source/CommonUI/Private/CommonUIEditorSettings.cpp b/Engine/Plugins/Experimental/CommonUI/Source/CommonUI/Private/CommonUIEditorSettings.cpp
--- a/Engine/Plugins/Experimental/CommonUI/Source/CommonUI/Private/CommonUIEditorSettings.cpp
+++ b/Engine/Plugins/Experimental/CommonUI/Source/CommonUI/Private/CommonUIEditorSettings.cpp
@@ -31,6 +31,7 @@ void UCommonUIEditorSettings::LoadEditorData()
 		TemplateTextStyleClass = TemplateTextStyle.LoadSynchronous();
 		TemplateButtonStyleClass = TemplateButtonStyle.LoadSynchronous();
 		TemplateBorderStyleClass = TemplateBorderStyle.LoadSynchronous();
+		DefaultHoldDataClass = DefaultHoldData.LoadSynchronous();
 
 		if (GUObjectArray.IsDisregardForGC(this))
 		{
@@ -46,6 +47,10 @@ void UCommonUIEditorSettings::LoadEditorData()
 			{
 				TemplateBorderStyleClass->AddToRoot();
 			}
+			if (DefaultHoldDataClass)
+			{
+				DefaultHoldDataClass->AddToRoot();
+			}
 		}
 
 		bDefaultDataLoaded = true;
@@ -72,4 +77,11 @@ const TSubclassOf<UCommonBorderStyle>& UCommonUIEditorSettings::GetTemplateBorde
 
 	return TemplateBorderStyleClass;
 }
+
+const TSubclassOf<UCommonUIHoldData>& UCommonUIEditorSettings::GetDefaultHoldData() const
+{
+	ensure(bDefaultDataLoaded);
+
+	return DefaultHoldDataClass;
+}
 #endif
diff --git a/Engine/Plugins/Experimental/CommonUI/Source/CommonUI/Private/CommonUITypes.cpp b/Engine/Plugins/Experimental/CommonUI/Source/CommonUI/Private/CommonUITypes.cpp
--- a/Engine/Plugins/Experimental/CommonUI/Source/CommonUI/Private/CommonUITypes.cpp
+++ b/Engine/Plugins/Experimental/CommonUI/Source/CommonUI/Private/CommonUITypes.cpp
@@ -27,6 +27,7 @@ FCommonInputTypeInfo::FCommonInputTypeInfo()
 	OverrideBrush.DrawAs = ESlateBrushDrawType::NoDrawType;
 	bActionRequiresHold = false;
 	HoldTime = 0.5f;
+	HoldRollbackTime = 0.0f;
 }
 
 FKey FCommonInputTypeInfo::GetKey() const
diff --git a/Engine/Plugins/Experimental/CommonUI/Source/CommonUI/Private/Input/UIActionRouterTypes.cpp b/Engine/Plugins/Experimental/CommonUI/Source/CommonUI/Private/Input/UIActionRouterTypes.cpp
--- a/Engine/Plugins/Experimental/CommonUI/Source/CommonUI/Private/Input/UIActionRouterTypes.cpp
+++ b/Engine/Plugins/Experimental/CommonUI/Source/CommonUI/Private/Input/UIActionRouterTypes.cpp
@@ -17,6 +17,8 @@
 #include "Input/UIActionBinding.h"
 #include "Widgets/SViewport.h"
 #include "CommonButtonBase.h"
+#include "CommonUIEditorSettings.h"
+#include "ICommonUIModule.h"
 
 DEFINE_LOG_CATEGORY(LogUIActionRouter);
 
@@ -117,6 +119,7 @@ FUIActionBinding::FUIActionBinding(const UWidget& InBoundWidget, const FBindUIAc
 	, InputAction(BindArgs.InputAction)
 {
 	OnHoldActionProgressed.Add(BindArgs.OnHoldActionProgressed);
+	OnHoldActionPressed.Add(BindArgs.OnHoldActionPressed);
 
 	const auto RegisterKeyMappingFunc = 
 		[this](const FUIActionKeyMapping& KeyMapping)
@@ -164,22 +167,31 @@ FUIActionBinding::FUIActionBinding(const UWidget& InBoundWidget, const FBindUIAc
 			ActionDisplayName = LegacyActionData->DisplayName;
 		}
 
+		FName CurrentGamepadName = UCommonUIActionRouterBase::Get(InBoundWidget)->GetInputSubsystem().GetCurrentGamepadName();
+		
 		// KB/M
 		const FCommonInputTypeInfo& LegacyMapping_KBM = LegacyActionData->GetInputTypeInfo(ECommonInputType::MouseAndKeyboard, FCommonInputDefaults::GamepadGeneric);
-		FUIActionKeyMapping KeyMapping_KBM(LegacyMapping_KBM.GetKey(), LegacyMapping_KBM.bActionRequiresHold ? LegacyMapping_KBM.HoldTime : 0.f);
-		RegisterKeyMappingFunc(KeyMapping_KBM);
+		FUIActionKeyMapping KeyMapping_KBM(LegacyMapping_KBM.GetKey(), LegacyMapping_KBM.bActionRequiresHold ? LegacyMapping_KBM.HoldTime : 0.f, LegacyMapping_KBM.bActionRequiresHold ? LegacyMapping_KBM.HoldRollbackTime : 0.f);
 
 		// Touch
 		const FCommonInputTypeInfo& LegacyMapping_Touch = LegacyActionData->GetInputTypeInfo(ECommonInputType::Touch, FCommonInputDefaults::GamepadGeneric);
-		FUIActionKeyMapping KeyMapping_Touch(LegacyMapping_Touch.GetKey(), LegacyMapping_Touch.bActionRequiresHold ? LegacyMapping_Touch.HoldTime : 0.f);
-		RegisterKeyMappingFunc(KeyMapping_Touch);
+		FUIActionKeyMapping KeyMapping_Touch(LegacyMapping_Touch.GetKey(), LegacyMapping_Touch.bActionRequiresHold ? LegacyMapping_Touch.HoldTime : 0.f, LegacyMapping_Touch.bActionRequiresHold ? LegacyMapping_Touch.HoldRollbackTime : 0.f);
 
 		// Gamepad
-		// Note: This is definitely a wonky and roundabout way to get the gamepad type, but given that it's for legacy fixup we'll let it go until this can be deleted
-		UCommonUIActionRouterBase* ActionRouter = UCommonUIActionRouterBase::Get(InBoundWidget);
-		check(ActionRouter);
-		const FCommonInputTypeInfo& LegacyMapping_Gamepad = LegacyActionData->GetInputTypeInfo(ECommonInputType::Gamepad, ActionRouter->GetInputSubsystem().GetCurrentGamepadName());
-		FUIActionKeyMapping KeyMapping_Gamepad(LegacyMapping_Gamepad.GetKey(), LegacyMapping_Gamepad.bActionRequiresHold ? LegacyMapping_Gamepad.HoldTime : 0.f);
+		const FCommonInputTypeInfo& LegacyMapping_Gamepad = LegacyActionData->GetInputTypeInfo(ECommonInputType::Gamepad, CurrentGamepadName);
+		FUIActionKeyMapping KeyMapping_Gamepad(LegacyMapping_Gamepad.GetKey(), LegacyMapping_Gamepad.bActionRequiresHold ? LegacyMapping_Gamepad.HoldTime : 0.f, LegacyMapping_Gamepad.bActionRequiresHold ? LegacyMapping_Gamepad.HoldRollbackTime : 0.f);
+
+		if (BindArgs.bForceHold && ICommonUIModule::GetEditorSettings().GetDefaultHoldData().GetDefaultObject())
+		{
+			UCommonUIHoldData* CommonUIHoldBehaviorValues = ICommonUIModule::GetEditorSettings().GetDefaultHoldData().GetDefaultObject();
+			
+			KeyMapping_KBM.HoldTime = KeyMapping_KBM.HoldTime > 0.f ? KeyMapping_KBM.HoldTime : CommonUIHoldBehaviorValues->KeyboardAndMouse.HoldTime;
+			KeyMapping_Touch.HoldTime = KeyMapping_Touch.HoldTime > 0.f ? KeyMapping_Touch.HoldTime : CommonUIHoldBehaviorValues->Touch.HoldTime;
+			KeyMapping_Gamepad.HoldTime = KeyMapping_Gamepad.HoldTime > 0.f ? KeyMapping_Gamepad.HoldTime : CommonUIHoldBehaviorValues->Gamepad.HoldTime;
+		}
+		
+		RegisterKeyMappingFunc(KeyMapping_KBM);
+		RegisterKeyMappingFunc(KeyMapping_Touch);
 		RegisterKeyMappingFunc(KeyMapping_Gamepad);
 	}
 
@@ -283,18 +295,25 @@ FString FUIActionBinding::ToDebugString() const
 
 void FUIActionBinding::BeginHold()
 {
+	if (HoldProgressRollbackTickerHandle.IsValid())
+	{
+		FTSTicker::GetCoreTicker().RemoveTicker(HoldProgressRollbackTickerHandle);
+		HoldProgressRollbackTickerHandle = nullptr;
+	}
+	
 	UE_LOG(LogUIActionRouter, VeryVerbose, TEXT("Hold pressed: %s"), *ActionName.ToString());
-
 	HoldStartTime = FPlatformTime::Seconds();
-	OnHoldActionProgressed.Broadcast(0.0f);
+	const float HoldStartPercent = HoldTime > 0.f ? HoldStartSecond / HoldTime : 0.f;
+	OnHoldActionProgressed.Broadcast(HoldStartPercent);
 }
 
 bool FUIActionBinding::UpdateHold(float TargetHoldTime)
 {
-	const float SecondsHeld = GetSecondsHeld();
-	UE_LOG(LogUIActionRouter, VeryVerbose, TEXT("Hold repeating: %s for %f"), *ActionName.ToString(), SecondsHeld);
-
-	const float HeldPercent = FMath::Clamp(SecondsHeld / TargetHoldTime, 0.f, 1.f);
+	CurrentHoldSecond = FMath::Clamp(GetSecondsHeld() + HoldStartSecond, 0.f, TargetHoldTime);
+	UE_LOG(LogUIActionRouter, VeryVerbose, TEXT("Hold repeating: %s for %f"), *ActionName.ToString(), CurrentHoldSecond);
+	
+	const float HeldPercent = FMath::Clamp(CurrentHoldSecond / TargetHoldTime, 0.f, 1.f);
+	
 	OnHoldActionProgressed.Broadcast(HeldPercent);
 
 	if (HeldPercent >= 1.f)
@@ -302,8 +321,10 @@ bool FUIActionBinding::UpdateHold(float TargetHoldTime)
 		CancelHold();
 		OnHoldActionProgressed.Broadcast(0.0f); //approximate CommonButton::OnActionComplete without adding extra bind
 		OnExecuteAction.ExecuteIfBound();
+		HoldStartSecond = 0.0f;
+		CurrentHoldSecond = 0.0f;
 
-		UE_LOG(LogUIActionRouter, VeryVerbose, TEXT("Hold repeating: Fired completed action! %s after %f"), *ActionName.ToString(), SecondsHeld);
+		UE_LOG(LogUIActionRouter, VeryVerbose, TEXT("Hold repeating: Fired completed action! %s after %f"), *ActionName.ToString(), CurrentHoldSecond);
 		return true;
 	}
 	return false;
@@ -324,11 +345,53 @@ void FUIActionBinding::CancelHold()
 	HoldStartTime = -1.0;
 }
 
+void FUIActionBinding::BeginRollback(float TargetHoldRollbackTime, float InHoldTime, FUIActionBindingHandle BindingHandle)
+{
+	HoldRollbackMultiplier = InHoldTime / TargetHoldRollbackTime;
+	HoldTime = InHoldTime;
+
+	if (TSharedPtr<FUIActionBinding> InBinding = FindBinding(BindingHandle))
+	{
+		HoldProgressRollbackTickerHandle = FTSTicker::GetCoreTicker().AddTicker(FTickerDelegate::CreateLambda([InBinding](float DeltaTime)
+			{
+				const float HoldRollbackTime = InBinding->HoldTime / InBinding->HoldRollbackMultiplier;
+				InBinding->CurrentHoldSecond = FMath::Clamp(InBinding->CurrentHoldSecond - (DeltaTime * InBinding->HoldRollbackMultiplier), 0.f, HoldRollbackTime);
+				const float HoldRollbackPercent = FMath::Clamp(InBinding->CurrentHoldSecond / InBinding->HoldTime, 0.f, 1.f);
+				InBinding->OnHoldActionProgressed.Broadcast(HoldRollbackPercent);
+				InBinding->HoldStartSecond = InBinding->CurrentHoldSecond;
+
+				if (HoldRollbackPercent <= 0.f)
+				{
+					FTSTicker::GetCoreTicker().RemoveTicker(InBinding->HoldProgressRollbackTickerHandle);
+					InBinding->HoldProgressRollbackTickerHandle = nullptr;
+					InBinding->HoldStartSecond = 0.f;
+					InBinding->CurrentHoldSecond = 0.f;
+						
+					return false;
+				}
+
+				return true;
+			}));
+	}
+}
+
 bool FUIActionBinding::IsHoldActive() const
 {
 	return HoldStartTime >= 0.0;
 }
 
+void FUIActionBinding::ResetHold()
+{
+	if (HoldProgressRollbackTickerHandle.IsValid())
+	{
+		FTSTicker::GetCoreTicker().RemoveTicker(HoldProgressRollbackTickerHandle);
+		HoldProgressRollbackTickerHandle = nullptr;
+	}
+	HoldStartSecond = 0.0f;
+	CurrentHoldSecond = 0.0f;
+}
+
+
 //////////////////////////////////////////////////////////////////////////
 // FBindUIActionArgs
 //////////////////////////////////////////////////////////////////////////
@@ -388,6 +451,14 @@ void FUIActionBindingHandle::Unregister()
 	RegistrationId = INDEX_NONE;
 }
 
+void FUIActionBindingHandle::ResetHold()
+{
+	if (TSharedPtr<FUIActionBinding> Binding = FUIActionBinding::FindBinding(*this))
+	{
+		Binding->ResetHold();
+	}
+}
+
 FName FUIActionBindingHandle::GetActionName() const
 {
 	if (TSharedPtr<const FUIActionBinding> Binding = FUIActionBinding::FindBinding(*this))
@@ -499,6 +570,7 @@ EProcessHoldActionResult FActionRouterBindingCollection::ProcessHoldInput(ECommo
 						{
 							Binding->BeginHold();
 							ProcessResult = EProcessHoldActionResult::Handled;
+							Binding->OnHoldActionPressed.Broadcast();
 						}
 						else if (Binding->IsHoldActive())
 						{
@@ -514,8 +586,17 @@ EProcessHoldActionResult FActionRouterBindingCollection::ProcessHoldInput(ECommo
 								const float SecondsHeld = Binding->GetSecondsHeld();
 
 								UE_LOG(LogUIActionRouter, VeryVerbose, TEXT("Hold released: %s after %f"), *Binding->ActionName.ToString(), SecondsHeld);
-								Binding->CancelHold();
-								Binding->OnHoldActionProgressed.Broadcast(0.0f);
+
+								if (HoldMapping.HoldRollbackTime <= UE_SMALL_NUMBER)
+								{
+									Binding->OnHoldActionProgressed.Broadcast(0.f);
+									Binding->CancelHold();
+								}
+								else
+								{
+									Binding->CancelHold();
+									Binding->BeginRollback(HoldMapping.HoldRollbackTime, HoldMapping.HoldTime, Binding->Handle);
+								}
 
 								static const float PressToHoldThreshold = 0.25f;
 								if (SecondsHeld <= PressToHoldThreshold && SecondsHeld < HoldMapping.HoldTime)
diff --git a/Engine/Plugins/Experimental/CommonUI/Source/CommonUI/Public/CommonButtonBase.h b/Engine/Plugins/Experimental/CommonUI/Source/CommonUI/Public/CommonButtonBase.h
--- a/Engine/Plugins/Experimental/CommonUI/Source/CommonUI/Public/CommonButtonBase.h
+++ b/Engine/Plugins/Experimental/CommonUI/Source/CommonUI/Public/CommonButtonBase.h
@@ -3,12 +3,14 @@
 #pragma once
 
 #include "Binding/States/WidgetStateRegistration.h"
+#include "CommonInputBaseTypes.h"
 #include "CommonUserWidget.h"
 #include "Components/Button.h"
 #include "Engine/DataTable.h"
 #include "FieldNotification/WidgetEventField.h"
 #include "Types/ISlateMetaData.h"
 #include "CommonInputModeTypes.h"
+#include "Containers/Ticker.h"
 #include "CommonButtonBase.generated.h"
 
 class UCommonButtonBase;
@@ -49,7 +51,6 @@ public:
 	}
 };
 
-
 /* ---- All properties must be EditDefaultsOnly, BlueprintReadOnly !!! -----
  *       we return the CDO to blueprints, so we cannot allow any changes (blueprint doesn't support const variables)
  */
@@ -508,6 +509,10 @@ protected:
 	UFUNCTION()
 	virtual void OnInputMethodChanged(ECommonInputType CurrentInputType);
 
+	/** If HoldData is valid, assigns its values to Keyboard and Mouse, Gamepad and Touch, based off the Current Input Type. */
+	UFUNCTION()
+	virtual void UpdateHoldData(ECommonInputType CurrentInputType);
+
 	/** Associates this button at its priority with the given key */
 	virtual void BindTriggeringInputActionToClick();
 
@@ -633,16 +638,35 @@ protected:
 	/** Callback fired continously during hold interactions */
 	UFUNCTION(BlueprintImplementableEvent, Category = "Common Button")
 	void OnActionProgress(float HeldPercent);
-	
+
+	/**
+	 * By default, if bRequiresHold is true, the bound Input Action will be forced to require hold as well.
+	 * The Bound Input action will use the row's hold values regardless of the row's bRequiresHold state.
+	 */
+	UFUNCTION()
+	virtual bool GetConvertInputActionToHold();
+
+	/** Bound to the hold progress of the bound key from the input action */
 	UFUNCTION()
 	virtual void NativeOnActionProgress(float HeldPercent);
 
+	/** Bound to the hold progress not related to the bound key */
+	UFUNCTION()
+	virtual bool NativeOnHoldProgress(float DeltaTime);
+	
+	/** Bound to the hold progress rollback not related to the bound key */
+	UFUNCTION()
+	virtual bool NativeOnHoldProgressRollback(float DeltaTime);
+	
 	/** Callback fired when hold events complete */
 	UFUNCTION(BlueprintImplementableEvent, Category = "Common Button")
 	void OnActionComplete();
 
 	UFUNCTION()
 	virtual void NativeOnActionComplete();
+	
+	UFUNCTION()
+	virtual void HoldReset();
 
 	virtual bool GetButtonAnalyticInfo(FString& ButtonName, FString& ABTestName, FString& ExtraData) const;
 
@@ -739,6 +763,14 @@ protected:
 	UPROPERTY(EditAnywhere, BlueprintReadOnly, Category = Input, meta = (ExposeOnSpawn = true))
 	uint8 bShouldUseFallbackDefaultInputAction:1;
 
+	/** True if this button should have a press and hold behavior, triggering the click when the specified hold time is met */
+	UPROPERTY(EditAnywhere, BlueprintReadOnly, Category = "Input|Hold", meta = (ExposeOnSpawn = true))
+	uint8 bRequiresHold:1;
+
+	/** Press and Hold values used for Keyboard and Mouse, Gamepad and Touch, depending on the current input type */
+	UPROPERTY(EditAnywhere, BlueprintReadOnly, Category = "Input|Hold", meta = (EditCondition="bRequiresHold", ExposeOnSpawn = true))
+	TSubclassOf<UCommonUIHoldData> HoldData;
+
 private:
 
 	/** True if this button is currently selected */
@@ -771,8 +803,8 @@ public:
 	int32 InputPriority;
 
 	/** 
-	 *	The input action that is bound to this button. The common input manager will trigger this button to 
-	 *	click if the action was pressed 
+	 * The input action that is bound to this button. The common input manager will trigger this button to 
+	 * click if the action was pressed 
 	 */
 	UPROPERTY(EditAnywhere, BlueprintReadOnly, Category = Input, meta = (ExposeOnSpawn = true, RowType = "/Script/CommonUI.CommonInputActionDataBase"))
 	FDataTableRowHandle TriggeringInputAction;
@@ -815,6 +847,25 @@ protected:
 	FCommonButtonBaseClicked OnButtonBaseUnhovered;
 
 	FUIActionBindingHandle TriggeringBindingHandle;
+	
+	/** Press and hold time in seconds */
+	float HoldTime;
+
+	/**
+	 * Time (in seconds) for hold progress to go from 1.0 (completed) to 0.0.
+	 * Used when the press and hold is interrupted.
+	 * If set to 0, there will be no rollback and the hold progress will reset immediately.
+	 */
+	float HoldRollbackTime;
+
+	/** Current hold time for this button */
+	float CurrentHoldTime;
+
+	/** Handle for ticker spawned for press and hold */
+	FTSTicker::FDelegateHandle HoldTickerHandle;
+
+	/** Handle for ticker spawned for button hold rollback */
+	FTSTicker::FDelegateHandle HoldProgressRollbackTickerHandle;
 
 private:
 	friend class UCommonButtonGroupBase;
diff --git a/Engine/Plugins/Experimental/CommonUI/Source/CommonUI/Public/CommonUIEditorSettings.h b/Engine/Plugins/Experimental/CommonUI/Source/CommonUI/Public/CommonUIEditorSettings.h
--- a/Engine/Plugins/Experimental/CommonUI/Source/CommonUI/Public/CommonUIEditorSettings.h
+++ b/Engine/Plugins/Experimental/CommonUI/Source/CommonUI/Public/CommonUIEditorSettings.h
@@ -28,6 +28,9 @@ public:
 
 	const TSubclassOf<UCommonBorderStyle>& GetTemplateBorderStyle() const;
 
+	/** Default Hold Data */
+	const TSubclassOf<UCommonUIHoldData>& GetDefaultHoldData() const;
+
 private:
 	void LoadEditorData();
 #endif
@@ -41,6 +44,13 @@ private:
 	UPROPERTY(config, EditAnywhere, Category = "Buttons")
 	TSoftClassPtr<UCommonButtonStyle> TemplateButtonStyle;
 
+	/**
+	 * Newly created CommonButton widgets will use these hold values by default if bRequiresHold is true.
+	 * Inherits from UCommonUIHoldData.
+	 */
+	UPROPERTY(EditAnywhere, Config, Category = "Buttons|Hold")
+	TSoftClassPtr<UCommonUIHoldData> DefaultHoldData;
+
 	/** Newly created CommonBorder Widgets will use this style. */
 	UPROPERTY(config, EditAnywhere, Category = "Border")
 	TSoftClassPtr<UCommonBorderStyle> TemplateBorderStyle;
@@ -54,6 +64,9 @@ private:
 
 	UPROPERTY(Transient)
 	TSubclassOf<UCommonBorderStyle> TemplateBorderStyleClass;
+
+	UPROPERTY(Transient)
+	TSubclassOf<UCommonUIHoldData> DefaultHoldDataClass;
 #endif
 private:
 	bool bDefaultDataLoaded;
diff --git a/Engine/Plugins/Experimental/CommonUI/Source/CommonUI/Public/CommonUITypes.h b/Engine/Plugins/Experimental/CommonUI/Source/CommonUI/Public/CommonUITypes.h
--- a/Engine/Plugins/Experimental/CommonUI/Source/CommonUI/Public/CommonUITypes.h
+++ b/Engine/Plugins/Experimental/CommonUI/Source/CommonUI/Public/CommonUITypes.h
@@ -64,8 +64,16 @@ public:
 	bool bActionRequiresHold;
 
 	/** The hold time in seconds */
-	UPROPERTY(EditAnywhere, Category = "CommonInput", meta = (EditCondition = "bActionRequiresHold"))
+	UPROPERTY(EditAnywhere, Category = "CommonInput", meta = (EditCondition = "bActionRequiresHold", ClampMin = "0.0", UIMin = "0.0"))
 	float HoldTime;
+
+	/**
+	 *	Time (in seconds) for hold progress to go from 1.0 (completed) to 0.0.
+	 *	If the hold interaction was interrupted, then hold progress starts to roll back decreasing its value.
+	 *	Set to 0.0 to disable the rollback functionality.
+	 */
+	UPROPERTY(EditAnywhere, Category = "CommonInput", meta = (EditCondition = "bActionRequiresHold", ClampMin = "0", UIMin = "0", ClampMax = "10.0", UIMax = "10"))
+	float HoldRollbackTime;
 	
 	/** Override the brush specified by the Key Display Data  */
 	UPROPERTY(EditAnywhere, Category = "CommonInput")
@@ -77,6 +85,7 @@ public:
 			OverrrideState == Other.OverrrideState &&
 			bActionRequiresHold == Other.bActionRequiresHold &&
 			HoldTime == Other.HoldTime &&
+			HoldRollbackTime == Other.HoldRollbackTime &&
 			OverrideBrush == Other.OverrideBrush;
 	}
 
diff --git a/Engine/Plugins/Experimental/CommonUI/Source/CommonUI/Public/Input/CommonUIInputSettings.h b/Engine/Plugins/Experimental/CommonUI/Source/CommonUI/Public/Input/CommonUIInputSettings.h
--- a/Engine/Plugins/Experimental/CommonUI/Source/CommonUI/Public/Input/CommonUIInputSettings.h
+++ b/Engine/Plugins/Experimental/CommonUI/Source/CommonUI/Public/Input/CommonUIInputSettings.h
@@ -15,9 +15,10 @@ struct FUIActionKeyMapping
 
 public:
 	FUIActionKeyMapping() {}
-	FUIActionKeyMapping(FKey InKey, float InHoldTime)
+	FUIActionKeyMapping(FKey InKey, float InHoldTime, float InHoldRollbackTime)
 		: Key(InKey)
 		, HoldTime(InHoldTime)
+		, HoldRollbackTime(InHoldRollbackTime)
 	{}
 
 	/** A key that triggers this action */
@@ -27,6 +28,10 @@ public:
 	/** How long must the key be held down for the action to be executed? */
 	UPROPERTY(EditAnywhere, Config, Category = "UI Action Key Mapping")
 	float HoldTime = 0.f;
+	
+	/** Time (in seconds) for hold progress to go from 1.0 (completed) to 0.0. */
+    UPROPERTY(EditAnywhere, Config, Category = "UI Action Key Mapping")
+    float HoldRollbackTime = 0.f;
 
 	//@todo DanH: Is this actually wanted/needed at all? Do we really want different bindings for the same thing on Platform A vs Platform B, for instance?
 	///** 
diff --git a/Engine/Plugins/Experimental/CommonUI/Source/CommonUI/Public/Input/CommonUIInputTypes.h b/Engine/Plugins/Experimental/CommonUI/Source/CommonUI/Public/Input/CommonUIInputTypes.h
--- a/Engine/Plugins/Experimental/CommonUI/Source/CommonUI/Public/Input/CommonUIInputTypes.h
+++ b/Engine/Plugins/Experimental/CommonUI/Source/CommonUI/Public/Input/CommonUIInputTypes.h
@@ -75,6 +75,9 @@ struct COMMONUI_API FBindUIActionArgs
 	/** Whether this binding can/should be displayed in a CommonActionBar (if one exists) */
 	bool bDisplayInActionBar = true;
 
+	/** True implies we will add default hold times if the current action is not a hold action */
+	bool bForceHold = false;
+
 	/** Optional display name to associate with this binding instead of the default */
 	FText OverrideDisplayName;
 
@@ -83,6 +86,10 @@ struct COMMONUI_API FBindUIActionArgs
 	/** If the bound action has any hold mappings, this will fire each frame while held. Has no bearing on actual execution and wholly irrelevant for non-hold actions */
 	DECLARE_DELEGATE_OneParam(FOnHoldActionProgressed, float);
 	FOnHoldActionProgressed OnHoldActionProgressed;
+
+	/** If the bound action has any hold mappings, this will fire when the hold begins. Has no bearing on actual execution and wholly irrelevant for non-hold actions */
+	DECLARE_DELEGATE(FOnHoldActionPressed);
+	FOnHoldActionPressed OnHoldActionPressed;
 };
 
 #if UE_ENABLE_INCLUDE_ORDER_DEPRECATED_IN_5_2
diff --git a/Engine/Plugins/Experimental/CommonUI/Source/CommonUI/Public/Input/UIActionBinding.h b/Engine/Plugins/Experimental/CommonUI/Source/CommonUI/Public/Input/UIActionBinding.h
--- a/Engine/Plugins/Experimental/CommonUI/Source/CommonUI/Public/Input/UIActionBinding.h
+++ b/Engine/Plugins/Experimental/CommonUI/Source/CommonUI/Public/Input/UIActionBinding.h
@@ -45,8 +45,10 @@ struct COMMONUI_API FUIActionBinding
 	void BeginHold();
 	bool UpdateHold(float TargetHoldTime);
 	void CancelHold();
+	void BeginRollback(float TargetHoldRollbackTime, float HoldTime, FUIActionBindingHandle BindingHandle);
 	double GetSecondsHeld() const;
 	bool IsHoldActive() const;
+	void ResetHold();
 
 	FName ActionName;
 	EInputEvent InputEvent;
@@ -70,14 +72,31 @@ struct COMMONUI_API FUIActionBinding
 	FOnHoldActionProgressedMulticast OnHoldActionProgressed;
 
 	// @TODO: Rename non-legacy in 5.3. We no longer have any active plans to remove data tables in CommonUI.
+
 	FDataTableRowHandle LegacyActionTableRow;
 
 	TWeakObjectPtr<const UInputAction> InputAction;
 
 private:
 	FUIActionBinding(const UWidget& InBoundWidget, const FBindUIActionArgs& BindArgs);
-	
+
+	// At what time in seconds did the hold start?
 	double HoldStartTime = -1.0;
+	
+	// At what second will the hold start?
+	double HoldStartSecond = 0.0;
+
+	// At what second is the hold progress at?
+	double CurrentHoldSecond = 0.0;
+
+	// Multiplier for the time (in seconds) for hold progress to go from 1.0 (completed) to 0.0.
+	double HoldRollbackMultiplier = 1.0;
+
+	// Target time (in seconds) for the hold progress to go from 0.0 to 1.0 (completed).
+	double HoldTime = 0.0;
+	
+	// Handle for ticker spawned for button hold rollback
+	FTSTicker::FDelegateHandle HoldProgressRollbackTickerHandle;
 
 	static int32 IdCounter;
 	static TMap<FUIActionBindingHandle, TSharedPtr<FUIActionBinding>> AllRegistrationsByHandle;
diff --git a/Engine/Plugins/Experimental/CommonUI/Source/CommonUI/Public/Input/UIActionBindingHandle.h b/Engine/Plugins/Experimental/CommonUI/Source/CommonUI/Public/Input/UIActionBindingHandle.h
--- a/Engine/Plugins/Experimental/CommonUI/Source/CommonUI/Public/Input/UIActionBindingHandle.h
+++ b/Engine/Plugins/Experimental/CommonUI/Source/CommonUI/Public/Input/UIActionBindingHandle.h
@@ -17,6 +17,9 @@ public:
 	bool IsValid() const;
 	void Unregister();
 
+	/** Calls ResetHold() on the Action Binding, which in turn resets the hold progress to 0.0  */
+	void ResetHold();
+
 	FName GetActionName() const;
 
 	FText GetDisplayName() const;
