Back out changelist 24878774

Original Description

Updated file by manually merging.

diff --git a/Engine/Plugins/Experimental/CommonUI/Source/CommonUI/Private/CommonButtonBase.cpp b/Engine/Plugins/Experimental/CommonUI/Source/CommonUI/Private/CommonButtonBase.cpp
--- a/Engine/Plugins/Experimental/CommonUI/Source/CommonUI/Private/CommonButtonBase.cpp
+++ b/Engine/Plugins/Experimental/CommonUI/Source/CommonUI/Private/CommonButtonBase.cpp
@@ -534,32 +534,14 @@ void UCommonButtonBase::BindTriggeringInputActionToClick()
 
 void UCommonButtonBase::BindTriggeringInputActionToClick()
 {
-	if (CommonUI::IsEnhancedInputSupportEnabled() && TriggeringEnhancedInputAction)
-	{
-		FBindUIActionArgs BindArgs(TriggeringEnhancedInputAction, false, FSimpleDelegate::CreateUObject(this, &UCommonButtonBase::HandleTriggeringActionCommited));
-		BindArgs.OnHoldActionProgressed.BindUObject(this, &UCommonButtonBase::NativeOnActionProgress);
-		BindArgs.OnHoldActionPressed.BindUObject(this, &UCommonButtonBase::HoldReset);
-        BindArgs.bIsPersistent = bIsPersistentBinding;
-        BindArgs.bForceHold = GetConvertInputActionToHold()
-
-		BindArgs.InputMode = InputModeOverride;
-
-		TriggeringBindingHandle = RegisterUIActionBinding(BindArgs);
-
-		return;
-	}
-
-	if (TriggeringInputAction.IsNull() || !TriggeredInputAction.IsNull())
-	{
-		return;
-	}
-
 	if (!TriggeringBindingHandle.IsValid())
 	{
 		FBindUIActionArgs BindArgs(TriggeringInputAction, false, FSimpleDelegate::CreateUObject(this, &UCommonButtonBase::HandleTriggeringActionCommited));
 		BindArgs.OnHoldActionProgressed.BindUObject(this, &UCommonButtonBase::NativeOnActionProgress);
+		BindArgs.OnHoldActionPressed.BindUObject(this, &UCommonButtonBase::HoldReset);
 		BindArgs.bIsPersistent = bIsPersistentBinding;
-
+		BindArgs.bForceHold = GetConvertInputActionToHold();
+		
 		BindArgs.InputMode = InputModeOverride;
 		
 		TriggeringBindingHandle = RegisterUIActionBinding(BindArgs);
