Missed file from 24882443

diff --git a/Engine/Plugins/Experimental/CommonUI/Source/CommonUI/Public/Input/UIActionBinding.h b/Engine/Plugins/Experimental/CommonUI/Source/CommonUI/Public/Input/UIActionBinding.h
--- a/Engine/Plugins/Experimental/CommonUI/Source/CommonUI/Public/Input/UIActionBinding.h
+++ b/Engine/Plugins/Experimental/CommonUI/Source/CommonUI/Public/Input/UIActionBinding.h
@@ -45,10 +45,8 @@ struct COMMONUI_API FUIActionBinding
 	void BeginHold();
 	bool UpdateHold(float TargetHoldTime);
 	void CancelHold();
-	void BeginRollback(float TargetHoldRollbackTime, float HoldTime, FUIActionBindingHandle BindingHandle);
 	double GetSecondsHeld() const;
 	bool IsHoldActive() const;
-	void ResetHold();
 
 	FName ActionName;
 	EInputEvent InputEvent;
@@ -72,31 +70,14 @@ struct COMMONUI_API FUIActionBinding
 	FOnHoldActionProgressedMulticast OnHoldActionProgressed;
 
 	// @TODO: Rename non-legacy in 5.3. We no longer have any active plans to remove data tables in CommonUI.
-
 	FDataTableRowHandle LegacyActionTableRow;
 
 	TWeakObjectPtr<const UInputAction> InputAction;
 
 private:
 	FUIActionBinding(const UWidget& InBoundWidget, const FBindUIActionArgs& BindArgs);
-
-	// At what time in seconds did the hold start?
-	double HoldStartTime = -1.0;
-	
-	// At what second will the hold start?
-	double HoldStartSecond = 0.0;
-
-	// At what second is the hold progress at?
-	double CurrentHoldSecond = 0.0;
-
-	// Multiplier for the time (in seconds) for hold progress to go from 1.0 (completed) to 0.0.
-	double HoldRollbackMultiplier = 1.0;
-
-	// Target time (in seconds) for the hold progress to go from 0.0 to 1.0 (completed).
-	double HoldTime = 0.0;
 	
-	// Handle for ticker spawned for button hold rollback
-	FTSTicker::FDelegateHandle HoldProgressRollbackTickerHandle;
+	double HoldStartTime = -1.0;
 
 	static int32 IdCounter;
 	static TMap<FUIActionBindingHandle, TSharedPtr<FUIActionBinding>> AllRegistrationsByHandle;
@@ -110,4 +91,4 @@ private:
 #if UE_ENABLE_INCLUDE_ORDER_DEPRECATED_IN_5_2
 #include "CoreMinimal.h"
 #include "Engine/EngineBaseTypes.h"
-#endif
+#endif
\ No newline at end of file
