Moved the underlying STextScroller implementation from CommonUI to Slate

diff --git a/Engine/Plugins/Experimental/CommonUI/Source/CommonUI/Private/CommonRichTextBlock.cpp b/Engine/Plugins/Experimental/CommonUI/Source/CommonUI/Private/CommonRichTextBlock.cpp
--- a/Engine/Plugins/Experimental/CommonUI/Source/CommonUI/Private/CommonRichTextBlock.cpp
+++ b/Engine/Plugins/Experimental/CommonUI/Source/CommonUI/Private/CommonRichTextBlock.cpp
@@ -5,6 +5,7 @@
 #include "Widgets/Text/SRichTextBlock.h"
 #include "Rendering/SlateRenderer.h"
 #include "Widgets/Text/STextBlock.h"
+#include "Widgets/Text/STextScroller.h"
 #include "Widgets/Images/SImage.h"
 #include "Widgets/Layout/SBox.h"
 
@@ -584,7 +585,7 @@ TSharedRef<SWidget> UCommonRichTextBlock::RebuildWidget()
 
 	MyTextScroller =
 		SNew(STextScroller)
-		.ScrollStyle(TextScrollStyle)
+		.ScrollOptions(TextScrollStyle->ToScrollOptions())
 		[
 			Super::RebuildWidget()
 		];
diff --git a/Engine/Plugins/Experimental/CommonUI/Source/CommonUI/Private/CommonTextBlock.cpp b/Engine/Plugins/Experimental/CommonUI/Source/CommonUI/Private/CommonTextBlock.cpp
--- a/Engine/Plugins/Experimental/CommonUI/Source/CommonUI/Private/CommonTextBlock.cpp
+++ b/Engine/Plugins/Experimental/CommonUI/Source/CommonUI/Private/CommonTextBlock.cpp
@@ -6,197 +6,21 @@
 #include "CommonWidgetPaletteCategories.h"
 #include "ICommonUIModule.h"
 #include "Widgets/Text/STextBlock.h"
+#include "Widgets/Text/STextScroller.h"
 #include "Layout/ArrangedChildren.h"
 #include "Layout/LayoutUtils.h"
 
 #include UE_INLINE_GENERATED_CPP_BY_NAME(CommonTextBlock)
 
-void STextScroller::Construct(const FArguments& InArgs)
+FTextScrollerOptions UCommonTextScrollStyle::ToScrollOptions() const
 {
-	ScrollStyle = InArgs._ScrollStyle;
-	ensure(ScrollStyle.IsValid());
-
-	// We don't tick, we use an active ticker.
-	SetCanTick(false);
-
-	ChildSlot
-	[
-		InArgs._Content.Widget
-	];
-}
-
-void STextScroller::OnArrangeChildren(const FGeometry& AllottedGeometry, FArrangedChildren& ArrangedChildren) const
-{
-	const EVisibility ChildVisibility = ChildSlot.GetWidget()->GetVisibility();
-	if (ArrangedChildren.Accepts(ChildVisibility))
-	{
-		const FVector2D ThisContentScale = GetContentScale();
-		const FMargin SlotPadding(LayoutPaddingWithFlow(GSlateFlowDirection, ChildSlot.GetPadding()));
-		const AlignmentArrangeResult XResult = AlignChild<Orient_Horizontal>(GSlateFlowDirection, AllottedGeometry.GetLocalSize().X, ChildSlot, SlotPadding, ThisContentScale.X, false);
-		const AlignmentArrangeResult YResult = AlignChild<Orient_Vertical>(AllottedGeometry.GetLocalSize().Y, ChildSlot, SlotPadding, ThisContentScale.Y);
-		const FVector2D MinimumSize = ChildSlot.GetWidget()->GetDesiredSize();
-
-		ArrangedChildren.AddWidget(ChildVisibility, AllottedGeometry.MakeChild(
-			ChildSlot.GetWidget(),
-			FVector2D(XResult.Offset, YResult.Offset),
-			FVector2D(XResult.Size, YResult.Size)
-		));
-	}
-}
-
-int32 STextScroller::OnPaint(const FPaintArgs& Args, const FGeometry& AllottedGeometry, const FSlateRect& MyCullingRect, FSlateWindowElementList& OutDrawElements, int32 LayerId, const FWidgetStyle& InWidgetStyle, bool bParentEnabled) const
-{
-	const_cast<STextScroller*>(this)->UpdateTickability(AllottedGeometry);
-
-	if (ScrollOffset != 0)
-	{
-		const float ScrollDirection = GSlateFlowDirection == EFlowDirection::LeftToRight ? -1 : 1;
-		const FGeometry ScrolledGeometry = AllottedGeometry.MakeChild(FSlateRenderTransform(FVector2D(ScrollOffset * ScrollDirection, 0)));
-		return SCompoundWidget::OnPaint(Args, ScrolledGeometry, MyCullingRect, OutDrawElements, LayerId, InWidgetStyle, bParentEnabled);
-	}
-	else
-	{
-		return SCompoundWidget::OnPaint(Args, AllottedGeometry, MyCullingRect, OutDrawElements, LayerId, InWidgetStyle, bParentEnabled);
-	}
-}
-
-void STextScroller::ResetScrollState()
-{
-	FontAlpha = 1.f;
-	TimeElapsed = 0.f;
-	ScrollOffset = 0.f;
-	// If suspended, make sure we stay suspended until receiving an explicit call to resume
-	if (IsScrollingEnabled())
-	{
-		ActiveState = EActiveState::EStart;
-	}
-	SetRenderOpacity(1.0f);
-}
-
-void STextScroller::StartScrolling()
-{
-	ActiveState = EActiveState::EStart;
-	ResetScrollState();
-}
-
-void STextScroller::SuspendScrolling()
-{
-	ActiveState = EActiveState::ESuspend;
-	ResetScrollState();
-}
-
-void STextScroller::UpdateTickability(const FGeometry& AllottedGeometry)
-{
-	const float VisibleWidth = AllottedGeometry.GetLocalSize().X;
-	const float DesiredWidth = VisibleWidth == 0 ? VisibleWidth : ChildSlot.GetWidget()->GetDesiredSize().X;
-
-	if (DesiredWidth > (VisibleWidth + 2) && IsScrollingEnabled())
-	{
-		if (!ActiveTimerHandle.IsValid())
-		{
-			ActiveTimerHandle = RegisterActiveTimer(0.0f, FWidgetActiveTimerDelegate::CreateSP(this, &STextScroller::OnScrollTextTick));
-			// If we need to scroll, then it's imperative that we arrange from the left, rather than fill, so that we flip to right aligning and scrolling
-			// to the right (potentially).
-			SetFlowDirectionPreference(EFlowDirectionPreference::Culture);
-			ChildSlot.HAlign(HAlign_Left);
-		}
-	}
-	else if (ActiveTimerHandle.IsValid())
-	{
-		UnRegisterActiveTimer(ActiveTimerHandle.ToSharedRef());
-		ActiveTimerHandle.Reset();
-		ResetScrollState();
-
-		// If we no longer need to scroll, just inherit the flow direction.
-		SetFlowDirectionPreference(EFlowDirectionPreference::Inherit);
-		ChildSlot.HAlign(HAlign_Fill);
-	}
-}
-
-EActiveTimerReturnType STextScroller::OnScrollTextTick(double CurrentTime, float DeltaTime)
-{
-	QUICK_SCOPE_CYCLE_COUNTER(STAT_UCommonTextBlock_OnTick);
-
-	const UCommonTextScrollStyle* TextScrollStyle = ScrollStyle.Get();
-	check(TextScrollStyle);
-
-	const float ContentSize = ChildSlot.GetWidget()->GetDesiredSize().X;
-	TimeElapsed += DeltaTime;
-
-	switch (ActiveState)
-	{
-	case EActiveState::EFadeIn:
-	{
-		FontAlpha = FMath::Clamp<float>(TimeElapsed / TextScrollStyle->FadeInDelay, 0.f, 1.f);
-		if (TimeElapsed >= TextScrollStyle->FadeInDelay)
-		{
-			FontAlpha = 1.f;
-			TimeElapsed = 0.f;
-			ScrollOffset = 0.f;
-			ActiveState = EActiveState::EStart;
-		}
-		break;
-	}
-	case EActiveState::EStart:
-	{
-		TimeElapsed = 0.f;
-		ScrollOffset = 0.f;
-		ActiveState = EActiveState::EStartWait;
-		break;
-	}
-	case EActiveState::EStartWait:
-	{
-		if (TimeElapsed >= TextScrollStyle->StartDelay)
-		{
-			TimeElapsed = 0.f;
-			ScrollOffset = 0.f;
-			ActiveState = EActiveState::EScrolling;
-		}
-		break;
-	}
-	case EActiveState::EScrolling:
-	{
-		ScrollOffset += TextScrollStyle->Speed * DeltaTime;
-		if ((ScrollOffset + GetCachedGeometry().GetLocalSize().X) >= ContentSize)
-		{
-			TimeElapsed = 0.0f;
-			ActiveState = EActiveState::EStop;
-		}
-		break;
-	}
-	case EActiveState::EStop:
-	{
-		TimeElapsed = 0.f;
-		ActiveState = EActiveState::EStopWait;
-		break;
-	}
-	case EActiveState::EStopWait:
-	{
-		if (TimeElapsed >= TextScrollStyle->EndDelay)
-		{
-			TimeElapsed = 0.f;
-			ActiveState = EActiveState::EFadeOut;
-		}
-		break;
-	}
-	case EActiveState::EFadeOut:
-	{
-		FontAlpha = 1.0f - FMath::Clamp<float>(TimeElapsed / TextScrollStyle->FadeOutDelay, 0.f, 1.f);
-		if (TimeElapsed >= TextScrollStyle->FadeOutDelay)
-		{
-			FontAlpha = 0.0f;
-			TimeElapsed = 0.0f;
-			ScrollOffset = 0.0f;
-			ActiveState = EActiveState::EFadeIn;
-		}
-		break;
-	}
-	}
-
-	SetRenderOpacity(FontAlpha);
-	Invalidate(EInvalidateWidgetReason::Paint);
-
-	return EActiveTimerReturnType::Continue;
+	FTextScrollerOptions ScrollOptions;
+	ScrollOptions.Speed = Speed;
+	ScrollOptions.StartDelay = StartDelay;
+	ScrollOptions.EndDelay = EndDelay;
+	ScrollOptions.FadeInDelay = FadeInDelay;
+	ScrollOptions.FadeOutDelay = FadeOutDelay;
+	return ScrollOptions;
 }
 
 // UCommonTextStyle
@@ -726,7 +550,7 @@ TSharedRef<SWidget> UCommonTextBlock::RebuildWidget()
 	// clang-format off
 	TextScroller = 
 		SNew(STextScroller)
-		.ScrollStyle(GetScrollStyleCDO())
+		.ScrollOptions(TextScrollStyle->ToScrollOptions())
 		[ 
 			Super::RebuildWidget() 
 		];
diff --git a/Engine/Plugins/Experimental/CommonUI/Source/CommonUI/Public/CommonTextBlock.h b/Engine/Plugins/Experimental/CommonUI/Source/CommonUI/Public/CommonTextBlock.h
--- a/Engine/Plugins/Experimental/CommonUI/Source/CommonUI/Public/CommonTextBlock.h
+++ b/Engine/Plugins/Experimental/CommonUI/Source/CommonUI/Public/CommonTextBlock.h
@@ -8,6 +8,9 @@
 #include "CommonTextBlock.generated.h"
 
 struct FTextBlockStyle;
+struct FTextScrollerOptions;
+
+class STextScroller;
 
 /* 
  * ---- All properties must be EditDefaultsOnly, BlueprintReadOnly !!! -----
@@ -89,6 +92,8 @@ class COMMONUI_API UCommonTextScrollStyle : public UObject
 	GENERATED_BODY()
 
 public:
+	FTextScrollerOptions ToScrollOptions() const;
+
 	UPROPERTY(EditAnywhere, BlueprintReadOnly, Category = "Properties")
 	float Speed;
 
@@ -105,62 +110,6 @@ public:
 	float FadeOutDelay;
 };
 
-class STextScroller : public SCompoundWidget
-{
-public:
-	SLATE_BEGIN_ARGS(STextScroller) {}
-		/** Slot for this designers content (optional) */
-		SLATE_DEFAULT_SLOT(FArguments, Content)
-
-		SLATE_ARGUMENT(const UCommonTextScrollStyle*, ScrollStyle)
-	SLATE_END_ARGS()
-
-public:
-	STextScroller()
-		: ActiveState(EActiveState::EStart)
-		, TimeElapsed(0.f)
-		, ScrollOffset(0.f)
-		, FontAlpha(1.f)
-	{
-	}
-
-	void Construct(const FArguments& InArgs);
-
-	virtual void OnArrangeChildren(const FGeometry& AllottedGeometry, FArrangedChildren& ArrangedChildren) const override;
-	virtual int32 OnPaint(const FPaintArgs& Args, const FGeometry& AllottedGeometry, const FSlateRect& MyCullingRect, FSlateWindowElementList& OutDrawElements, int32 LayerId, const FWidgetStyle& InWidgetStyle, bool bParentEnabled) const override;
-
-	void ResetScrollState();
-	void StartScrolling();
-	void SuspendScrolling();
-
-	bool IsScrollingEnabled() const { return ActiveState != EActiveState::ESuspend; }
-
-private:
-	void UpdateTickability(const FGeometry& AllottedGeometry);
-	EActiveTimerReturnType OnScrollTextTick(double CurrentTime, float DeltaTime);
-
-private:
-	enum class EActiveState : uint8
-	{
-		EFadeIn = 0,
-		EStart,
-		EStartWait,
-		EScrolling,
-		EStop,
-		EStopWait,
-		EFadeOut,
-		ESuspend,
-	} ActiveState;
-
-	float TimeElapsed;
-	float ScrollOffset;
-	float FontAlpha;
-
-	TSharedPtr<FActiveTimerHandle> ActiveTimerHandle;
-
-	TWeakObjectPtr<const UCommonTextScrollStyle> ScrollStyle;
-};
-
 /**
  * Text block with automatic scrolling for FX / large texts, also supports a larger set of default styling, & custom mobile scaling.
  */
diff --git a/Engine/Source/Runtime/Slate/Private/Widgets/Text/STextScroller.cpp b/Engine/Source/Runtime/Slate/Private/Widgets/Text/STextScroller.cpp
new file mode 100644
--- /dev/null
+++ b/Engine/Source/Runtime/Slate/Private/Widgets/Text/STextScroller.cpp
@@ -0,0 +1,207 @@
+// Copyright Epic Games, Inc. All Rights Reserved.
+
+#include "Widgets/Text/STextScroller.h"
+#include "Widgets/Text/STextBlock.h"
+
+#include "Layout/ArrangedChildren.h"
+#include "Layout/LayoutUtils.h"
+
+void STextScroller::Construct(const FArguments& InArgs)
+{
+	ScrollOptions = InArgs._ScrollOptions;
+
+	// We don't tick, we use an active ticker.
+	SetCanTick(false);
+
+	ChildSlot
+	[
+		InArgs._Content.Widget
+	];
+}
+
+void STextScroller::OnArrangeChildren(const FGeometry& AllottedGeometry, FArrangedChildren& ArrangedChildren) const
+{
+	const EVisibility ChildVisibility = ChildSlot.GetWidget()->GetVisibility();
+	if (ArrangedChildren.Accepts(ChildVisibility))
+	{
+		const FVector2D ThisContentScale = GetContentScale();
+		const FMargin SlotPadding(LayoutPaddingWithFlow(GSlateFlowDirection, ChildSlot.GetPadding()));
+		const AlignmentArrangeResult XResult = AlignChild<Orient_Horizontal>(GSlateFlowDirection, AllottedGeometry.GetLocalSize().X, ChildSlot, SlotPadding, ThisContentScale.X, false);
+		const AlignmentArrangeResult YResult = AlignChild<Orient_Vertical>(AllottedGeometry.GetLocalSize().Y, ChildSlot, SlotPadding, ThisContentScale.Y);
+		const FVector2D MinimumSize = ChildSlot.GetWidget()->GetDesiredSize();
+
+		ArrangedChildren.AddWidget(ChildVisibility, AllottedGeometry.MakeChild(
+			ChildSlot.GetWidget(),
+			FVector2D(XResult.Offset, YResult.Offset),
+			FVector2D(XResult.Size, YResult.Size)
+		));
+	}
+}
+
+int32 STextScroller::OnPaint(const FPaintArgs& Args, const FGeometry& AllottedGeometry, const FSlateRect& MyCullingRect, FSlateWindowElementList& OutDrawElements, int32 LayerId, const FWidgetStyle& InWidgetStyle, bool bParentEnabled) const
+{
+	{
+		STextScroller* MutableThis = const_cast<STextScroller*>(this);
+
+		const float VisibleWidth = AllottedGeometry.GetLocalSize().X;
+		const float DesiredWidth = VisibleWidth == 0.0f ? VisibleWidth : ChildSlot.GetWidget()->GetDesiredSize().X;
+
+		if (DesiredWidth > (VisibleWidth + 2.0f) && IsScrollingEnabled())
+		{
+			if (!ActiveTimerHandle.IsValid())
+			{
+				MutableThis->TickerState = ETickerState::StartTicking;
+				MutableThis->ActiveTimerHandle = MutableThis->RegisterActiveTimer(0.0f, FWidgetActiveTimerDelegate::CreateSP(MutableThis, &STextScroller::OnScrollTextTick));
+			}
+		}
+		else if (ActiveTimerHandle.IsValid())
+		{
+			MutableThis->TickerState = ETickerState::StopTicking;
+		}
+	}
+
+	if (ScrollOffset != 0)
+	{
+		const float ScrollDirection = GSlateFlowDirection == EFlowDirection::LeftToRight ? -1 : 1;
+		const FGeometry ScrolledGeometry = AllottedGeometry.MakeChild(FSlateRenderTransform(FVector2D(ScrollOffset * ScrollDirection, 0)));
+		return SCompoundWidget::OnPaint(Args, ScrolledGeometry, MyCullingRect, OutDrawElements, LayerId, InWidgetStyle, bParentEnabled);
+	}
+	else
+	{
+		return SCompoundWidget::OnPaint(Args, AllottedGeometry, MyCullingRect, OutDrawElements, LayerId, InWidgetStyle, bParentEnabled);
+	}
+}
+
+void STextScroller::ResetScrollState()
+{
+	FontAlpha = 1.f;
+	TimeElapsed = 0.f;
+	ScrollOffset = 0.f;
+	// If suspended, make sure we stay suspended until receiving an explicit call to resume
+	if (IsScrollingEnabled())
+	{
+		ActiveState = EActiveState::Start;
+	}
+	SetRenderOpacity(1.0f);
+}
+
+void STextScroller::StartScrolling()
+{
+	ActiveState = EActiveState::Start;
+	ResetScrollState();
+}
+
+void STextScroller::SuspendScrolling()
+{
+	ActiveState = EActiveState::Suspend;
+	ResetScrollState();
+}
+
+EActiveTimerReturnType STextScroller::OnScrollTextTick(double CurrentTime, float DeltaTime)
+{
+	QUICK_SCOPE_CYCLE_COUNTER(STAT_STextScroller_OnTick);
+
+	check(TickerState != ETickerState::None);
+	switch (TickerState)
+	{
+	case ETickerState::StartTicking:
+		// If we need to scroll, then it's imperative that we arrange from the left, rather than fill, so that we flip to right aligning and scrolling
+		// to the right (potentially).
+		SetFlowDirectionPreference(EFlowDirectionPreference::Culture);
+		ChildSlot.HAlign(HAlign_Left);
+
+		TickerState = ETickerState::Ticking;
+		return EActiveTimerReturnType::Continue; // Defer to the next tick so that the layout is correct
+
+	case ETickerState::StopTicking:
+		// If we no longer need to scroll, just inherit the flow direction.
+		SetFlowDirectionPreference(EFlowDirectionPreference::Inherit);
+		ChildSlot.HAlign(HAlign_Fill);
+
+		TickerState = ETickerState::None;
+		ResetScrollState();
+		ActiveTimerHandle.Reset();
+		return EActiveTimerReturnType::Stop; // Ticking is no longer required
+
+	default:
+		break;
+	}
+
+	const float ContentSize = ChildSlot.GetWidget()->GetDesiredSize().X;
+	TimeElapsed += DeltaTime;
+
+	switch (ActiveState)
+	{
+	case EActiveState::FadeIn:
+	{
+		FontAlpha = FMath::Clamp<float>(TimeElapsed / ScrollOptions.FadeInDelay, 0.f, 1.f);
+		if (TimeElapsed >= ScrollOptions.FadeInDelay)
+		{
+			FontAlpha = 1.f;
+			TimeElapsed = 0.f;
+			ScrollOffset = 0.f;
+			ActiveState = EActiveState::Start;
+		}
+		break;
+	}
+	case EActiveState::Start:
+	{
+		TimeElapsed = 0.f;
+		ScrollOffset = 0.f;
+		ActiveState = EActiveState::StartWait;
+		break;
+	}
+	case EActiveState::StartWait:
+	{
+		if (TimeElapsed >= ScrollOptions.StartDelay)
+		{
+			TimeElapsed = 0.f;
+			ScrollOffset = 0.f;
+			ActiveState = EActiveState::Scrolling;
+		}
+		break;
+	}
+	case EActiveState::Scrolling:
+	{
+		ScrollOffset += ScrollOptions.Speed * DeltaTime;
+		if ((ScrollOffset + GetCachedGeometry().GetLocalSize().X) >= ContentSize)
+		{
+			TimeElapsed = 0.0f;
+			ActiveState = EActiveState::Stop;
+		}
+		break;
+	}
+	case EActiveState::Stop:
+	{
+		TimeElapsed = 0.f;
+		ActiveState = EActiveState::StopWait;
+		break;
+	}
+	case EActiveState::StopWait:
+	{
+		if (TimeElapsed >= ScrollOptions.EndDelay)
+		{
+			TimeElapsed = 0.f;
+			ActiveState = EActiveState::FadeOut;
+		}
+		break;
+	}
+	case EActiveState::FadeOut:
+	{
+		FontAlpha = 1.0f - FMath::Clamp<float>(TimeElapsed / ScrollOptions.FadeOutDelay, 0.f, 1.f);
+		if (TimeElapsed >= ScrollOptions.FadeOutDelay)
+		{
+			FontAlpha = 0.0f;
+			TimeElapsed = 0.0f;
+			ScrollOffset = 0.0f;
+			ActiveState = EActiveState::FadeIn;
+		}
+		break;
+	}
+	}
+
+	SetRenderOpacity(FontAlpha);
+	Invalidate(EInvalidateWidgetReason::Paint);
+
+	return EActiveTimerReturnType::Continue;
+}
diff --git a/Engine/Source/Runtime/Slate/Public/Widgets/Text/STextScroller.h b/Engine/Source/Runtime/Slate/Public/Widgets/Text/STextScroller.h
new file mode 100644
--- /dev/null
+++ b/Engine/Source/Runtime/Slate/Public/Widgets/Text/STextScroller.h
@@ -0,0 +1,87 @@
+// Copyright Epic Games, Inc. All Rights Reserved.
+
+#pragma once
+
+#include "Layout/Clipping.h"
+#include "Widgets/SCompoundWidget.h"
+
+struct FTextScrollerOptions
+{
+	/** Speed (in Slate units) to scroll per-second */
+	float Speed = 25.0f;
+
+	/** Delay (in seconds) to pause before scrolling the text */
+	float StartDelay = 0.5f;
+
+	/** Delay (in seconds) to pause after scrolling the text */
+	float EndDelay = 0.5f;
+
+	/** Duration (in seconds) to fade in before scrolling the text */
+	float FadeInDelay = 0.0f;
+
+	/** Duration (in seconds) to fade out after scrolling the text */
+	float FadeOutDelay = 0.0f;
+};
+
+/**
+ * Utility to wrap a single-line text widget (STextBlock or SRichTextBlock) and provide support for auto-scrolling the text if it's longer than the available space.
+ */
+class SLATE_API STextScroller : public SCompoundWidget
+{
+public:
+	SLATE_BEGIN_ARGS(STextScroller)
+	{
+		_Clipping = EWidgetClipping::OnDemand;
+	}
+		/** Options to apply when scrolling the text */
+		SLATE_ARGUMENT(FTextScrollerOptions, ScrollOptions)
+		/** Slot for the text widget to be scrolled */
+		SLATE_DEFAULT_SLOT(FArguments, Content)
+	SLATE_END_ARGS()
+
+public:
+	void Construct(const FArguments& InArgs);
+
+	virtual void OnArrangeChildren(const FGeometry& AllottedGeometry, FArrangedChildren& ArrangedChildren) const override;
+	virtual int32 OnPaint(const FPaintArgs& Args, const FGeometry& AllottedGeometry, const FSlateRect& MyCullingRect, FSlateWindowElementList& OutDrawElements, int32 LayerId, const FWidgetStyle& InWidgetStyle, bool bParentEnabled) const override;
+
+	void ResetScrollState();
+	void StartScrolling();
+	void SuspendScrolling();
+
+	bool IsScrollingEnabled() const { return ActiveState != EActiveState::Suspend; }
+
+private:
+	EActiveTimerReturnType OnScrollTextTick(double CurrentTime, float DeltaTime);
+
+private:
+	enum class EActiveState : uint8
+	{
+		FadeIn = 0,
+		Start,
+		StartWait,
+		Scrolling,
+		Stop,
+		StopWait,
+		FadeOut,
+		Suspend,
+	};
+
+	enum class ETickerState : uint8
+	{
+		None = 0,
+		StartTicking,
+		Ticking,
+		StopTicking,
+	};
+
+	EActiveState ActiveState = EActiveState::Start;
+	ETickerState TickerState = ETickerState::None;
+	float TimeElapsed = 0.0f;
+	float ScrollOffset = 0.0f;
+	float FontAlpha = 1.0f;
+
+	TSharedPtr<FActiveTimerHandle> ActiveTimerHandle;
+
+	FTextScrollerOptions ScrollOptions;
+};
