Hold Support for Common Button Base:
- Added flag to make the button require hold to trigger action when focused and pressed
- Added optional Hold Behavior abstract class to specify Keyboard and Mouse, Gamepad and Touch hold behavior values.
- Modify Hold behavior for input bindings, including Rollback multiplier for a required hold input binding.
- Added Default Hold Data in the Common UI Editor settings, under the Button category.

diff --git a/Engine/Plugins/Experimental/CommonUI/Source/CommonInput/Public/CommonInputBaseTypes.h b/Engine/Plugins/Experimental/CommonUI/Source/CommonInput/Public/CommonInputBaseTypes.h
--- a/Engine/Plugins/Experimental/CommonUI/Source/CommonInput/Public/CommonInputBaseTypes.h
+++ b/Engine/Plugins/Experimental/CommonUI/Source/CommonInput/Public/CommonInputBaseTypes.h
@@ -77,6 +77,17 @@ struct FInputDeviceIdentifierPair
 	FString HardwareDeviceIdentifier;
 };
 
+/* Data values needed for Hold interaction per input type. */
+USTRUCT()
+struct FInputHoldData
+{
+	GENERATED_BODY()
+	UPROPERTY(EditDefaultsOnly, Category = "Properties", meta = (ClampMin = "0.0", UIMin = "0.0"))
+	float HoldTime;
+	UPROPERTY(EditDefaultsOnly, Category = "Properties", meta = (ClampMin = "0.0", UIMin = "0.0"))
+	float HoldRollbackTime;
+};
+
 /* Derive from this class to store the Input data. It is referenced in the Common Input Settings, found in the project settings UI. */
 UCLASS(Abstract, Blueprintable, ClassGroup = Input, meta = (Category = "Common Input"))
 class COMMONINPUT_API UCommonUIInputData : public UObject
@@ -100,6 +111,30 @@ public:
 	TObjectPtr<UInputAction> EnhancedInputBackAction;
 };
 
+/* Defines values for hold behavior per input type: for mouse Press and Hold, for gamepad focused Press and Hold, for touch Press and Hold. */
+UCLASS(Abstract, Blueprintable, ClassGroup = Input, meta = (Category = "Common Input"))
+class COMMONINPUT_API UCommonUIHoldData : public UObject
+{
+	GENERATED_BODY()
+public:
+	UCommonUIHoldData()
+	{
+		KeyboardAndMouse.HoldTime = 0.75f;
+		KeyboardAndMouse.HoldRollbackTime = 0.0f;
+		Gamepad.HoldTime = 0.75f;
+		Gamepad.HoldRollbackTime = 0.0f;
+		Touch.HoldTime = 0.75f;
+		Touch.HoldRollbackTime = 0.0f;
+	}
+	
+	UPROPERTY(EditDefaultsOnly, Category = "KeyboardAndMouse", meta = (ShowOnlyInnerProperties))
+	FInputHoldData KeyboardAndMouse;
+	UPROPERTY(EditDefaultsOnly, Category = "Gamepad", meta = (ShowOnlyInnerProperties))
+	FInputHoldData Gamepad;
+	UPROPERTY(EditDefaultsOnly, Category = "Touch", meta = (ShowOnlyInnerProperties))
+	FInputHoldData Touch;
+};
+
 /* Derive from this class to store the Input data. It is referenced in the Common Input Settings, found in the project settings UI. */
 UCLASS(Abstract, Blueprintable, ClassGroup = Input, meta = (Category = "Common Input"))
 class COMMONINPUT_API UCommonInputBaseControllerData : public UObject
diff --git a/Engine/Plugins/Experimental/CommonUI/Source/CommonUI/Private/CommonButtonBase.cpp b/Engine/Plugins/Experimental/CommonUI/Source/CommonUI/Private/CommonButtonBase.cpp
--- a/Engine/Plugins/Experimental/CommonUI/Source/CommonUI/Private/CommonButtonBase.cpp
+++ b/Engine/Plugins/Experimental/CommonUI/Source/CommonUI/Private/CommonButtonBase.cpp
@@ -24,6 +24,9 @@
 #include "InputAction.h"
 #include "Sound/SoundBase.h"
 #include "Styling/UMGCoreStyle.h"
+#include "CommonUITypes.h"
+#include "CommonUIPrivate.h"
+#include "Input/CommonUIActionRouterBase.h"
 
 #include UE_INLINE_GENERATED_CPP_BY_NAME(CommonButtonBase)
 
@@ -345,9 +348,13 @@ UCommonButtonBase::UCommonButtonBase(const FObjectInitializer& ObjectInitializer
 	, bTriggerClickedAfterSelection(false)
 	, bDisplayInputActionWhenNotInteractable(true)
 	, bShouldUseFallbackDefaultInputAction(true)
+	, bRequiresHold(false)
 	, bSelected(false)
 	, bButtonEnabled(true)
 	, bInteractionEnabled(true)
+	, HoldTime(0.f)
+	, HoldRollbackTime(0.f)
+	, CurrentHoldTime(0.f)
 {
 	SetIsFocusable(true);
 }
@@ -368,7 +375,15 @@ void UCommonButtonBase::PostLoad()
 	{
 		UCommonUIEditorSettings& Settings = ICommonUIModule::GetEditorSettings();
 		Settings.ConditionalPostLoad();
-		Style = Settings.GetTemplateButtonStyle();
+		
+		if (!Style && !bStyleNoLongerNeedsConversion)
+		{
+			Style = Settings.GetTemplateButtonStyle();
+		}
+		if (!HoldData)
+		{
+			HoldData = Settings.GetDefaultHoldData();
+		}
 	}
 
 	bStyleNoLongerNeedsConversion = true;
@@ -391,6 +406,10 @@ void UCommonButtonBase::OnCreationFromPalette()
 	{
 		Style = ICommonUIModule::GetEditorSettings().GetTemplateButtonStyle();
 	}
+	if (!HoldData)
+	{
+		HoldData = ICommonUIModule::GetEditorSettings().GetDefaultHoldData();
+	}
 	Super::OnCreationFromPalette();
 }
 
@@ -443,6 +462,11 @@ UCommonButtonInternalBase* UCommonButtonBase::ConstructInternalButton()
 
 void UCommonButtonBase::NativeConstruct()
 {
+	if (bRequiresHold && !HoldData)
+	{
+		HoldData = ICommonUIModule::GetEditorSettings().GetDefaultHoldData();
+	}
+	
 	BindTriggeringInputActionToClick();
 	BindInputMethodChangedDelegate();
 	UpdateInputActionWidget();
@@ -456,6 +480,17 @@ void UCommonButtonBase::NativeDestruct()
 
 	UnbindTriggeringInputActionToClick();
 	UnbindInputMethodChangedDelegate();
+
+	if (HoldTickerHandle.IsValid())
+	{
+		FTSTicker::GetCoreTicker().RemoveTicker(HoldTickerHandle);
+		HoldTickerHandle = nullptr;
+	}
+	if (HoldProgressRollbackTickerHandle.IsValid())
+	{
+		FTSTicker::GetCoreTicker().RemoveTicker(HoldProgressRollbackTickerHandle);
+		HoldProgressRollbackTickerHandle = nullptr;
+	}
 }
 
 void UCommonButtonBase::SetIsEnabled(bool bInIsEnabled)
@@ -496,6 +531,7 @@ void UCommonButtonBase::BindInputMethodChangedDelegate()
 	if (CommonInputSubsystem)
 	{
 		CommonInputSubsystem->OnInputMethodChangedNative.AddUObject(this, &UCommonButtonBase::OnInputMethodChanged);
+		UpdateHoldData(CommonInputSubsystem->GetDefaultInputType());
 	}
 }
 
@@ -511,7 +547,43 @@ void UCommonButtonBase::UnbindInputMethodChangedDelegate()
 void UCommonButtonBase::OnInputMethodChanged(ECommonInputType CurrentInputType)
 {
 	UpdateInputActionWidget();
+	UpdateHoldData(CurrentInputType);
+	HoldReset();
+	NativeOnActionProgress(0.f);
 	BP_OnInputMethodChanged(CurrentInputType);
+
+	if (TriggeringBindingHandle.IsValid())
+	{
+		TriggeringBindingHandle.ResetHold();
+	}
+}
+
+void UCommonButtonBase::UpdateHoldData(ECommonInputType CurrentInputType)
+{
+	if (HoldData && bRequiresHold)
+	{
+		UCommonUIHoldData* CommonUIHoldBehaviorValues = HoldData.GetDefaultObject();
+		if (CommonUIHoldBehaviorValues)
+		{
+			switch (CurrentInputType)
+			{
+			case ECommonInputType::MouseAndKeyboard:
+				HoldTime = CommonUIHoldBehaviorValues->KeyboardAndMouse.HoldTime;
+				HoldRollbackTime = CommonUIHoldBehaviorValues->KeyboardAndMouse.HoldRollbackTime;
+				break;
+			case ECommonInputType::Gamepad:
+				HoldTime = CommonUIHoldBehaviorValues->Gamepad.HoldTime;
+				HoldRollbackTime = CommonUIHoldBehaviorValues->Gamepad.HoldRollbackTime;
+				break;
+			case ECommonInputType::Touch:
+				HoldTime = CommonUIHoldBehaviorValues->Touch.HoldTime;
+				HoldRollbackTime = CommonUIHoldBehaviorValues->Touch.HoldRollbackTime;
+				break;
+			default:
+				break;
+			}
+		}
+	}
 }
 
 void UCommonButtonBase::BindTriggeringInputActionToClick()
@@ -538,7 +610,9 @@ void UCommonButtonBase::BindTriggeringInputActionToClick()
 	{
 		FBindUIActionArgs BindArgs(TriggeringInputAction, false, FSimpleDelegate::CreateUObject(this, &UCommonButtonBase::HandleTriggeringActionCommited));
 		BindArgs.OnHoldActionProgressed.BindUObject(this, &UCommonButtonBase::NativeOnActionProgress);
+		BindArgs.OnHoldActionPressed.BindUObject(this, &UCommonButtonBase::HoldReset);
 		BindArgs.bIsPersistent = bIsPersistentBinding;
+		BindArgs.bForceHold = GetConvertInputActionToHold();
 
 		BindArgs.InputMode = InputModeOverride;
 		
@@ -564,6 +638,8 @@ void UCommonButtonBase::UnbindTriggeringInputActionToClick()
 	{
 		TriggeringBindingHandle.Unregister();
 	}
+	
+	CurrentHoldTime = 0.f;
 }
 
 void UCommonButtonBase::HandleTriggeringActionCommited(bool& bPassthrough)
@@ -714,6 +790,11 @@ void UCommonButtonBase::SetIsInteractableWhenSelected(bool bInInteractableWhenSe
 	}
 }
 
+bool UCommonButtonBase::GetConvertInputActionToHold()
+{
+	return bRequiresHold;
+}
+
 void UCommonButtonBase::NativeOnActionProgress(float HeldPercent)
 {
 	if (InputActionWidget)
@@ -723,6 +804,65 @@ void UCommonButtonBase::NativeOnActionProgress(float HeldPercent)
 	OnActionProgress(HeldPercent);
 }
 
+bool UCommonButtonBase::NativeOnHoldProgress(float DeltaTime)
+{
+	if (HoldTime > UE_SMALL_NUMBER)
+	{
+		CurrentHoldTime += FMath::Clamp(DeltaTime, 0.f, HoldTime);
+		const float HeldPercent = FMath::Clamp(CurrentHoldTime / HoldTime, 0.f, 1.f);
+		NativeOnActionProgress(HeldPercent);
+		if (HeldPercent >= 1.f)
+		{
+			HandleTriggeringActionCommited();
+			HoldReset();
+			return false;
+		}
+		
+		return true;
+	}
+	HoldReset();
+	
+	return false;
+}
+
+bool UCommonButtonBase::NativeOnHoldProgressRollback(float DeltaTime)
+{
+	if (HoldTime > UE_SMALL_NUMBER && HoldRollbackTime > UE_SMALL_NUMBER)
+	{
+		const float HoldRollbackMultiplier = HoldTime / HoldRollbackTime;
+		CurrentHoldTime = FMath::Clamp(CurrentHoldTime - (DeltaTime * HoldRollbackMultiplier), 0.f, HoldRollbackTime);
+		const float HoldRollbackPercent = FMath::Clamp(CurrentHoldTime / HoldTime, 0.f, 1.f);
+		NativeOnActionProgress(HoldRollbackPercent);
+		if (HoldRollbackPercent <= 0.f)
+		{
+			FTSTicker::GetCoreTicker().RemoveTicker(HoldProgressRollbackTickerHandle);
+			HoldProgressRollbackTickerHandle = nullptr;
+		
+			return false;
+		}
+		
+		return true;
+	}
+	HoldReset();
+	
+	return false;
+}
+
+void UCommonButtonBase::HoldReset()
+{
+	if (HoldTickerHandle.IsValid())
+	{
+		FTSTicker::GetCoreTicker().RemoveTicker(HoldTickerHandle);
+		HoldTickerHandle = nullptr;
+	}
+	if (HoldProgressRollbackTickerHandle.IsValid())
+	{
+		FTSTicker::GetCoreTicker().RemoveTicker(HoldProgressRollbackTickerHandle);
+		HoldProgressRollbackTickerHandle = nullptr;
+	}
+	CurrentHoldTime = 0.0f;
+}
+
 void UCommonButtonBase::NativeOnActionComplete()
 {
 	if (InputActionWidget)
@@ -1161,6 +1301,13 @@ void UCommonButtonBase::HandleButtonClicked()
 	// Guard against this case.
 	if (IsInteractionEnabled())
 	{
+		// @TODO: Current click rejection method relies on click hold time, this can be refined. See NativeOnHoldProgress.
+		// Also gamepad can indirectly trigger this method, so don't guard against pressed
+    	if (bRequiresHold && CurrentHoldTime < HoldTime && IsHovered())
+    	{
+    		return;
+    	}
+    		
 		if (bTriggerClickedAfterSelection)
 		{
 			SetIsSelected(!bSelected, false);
@@ -1173,6 +1320,7 @@ void UCommonButtonBase::HandleButtonClicked()
 		}
 
 		ExecuteTriggeredInput();
+		HoldReset();
 	}
 }
 
@@ -1210,6 +1358,21 @@ void UCommonButtonBase::HandleButtonPressed()
 		// Simulate hover events when using touch input
 		NativeOnHovered();
 	}
+
+	if (bRequiresHold && HoldTime > 0.f)
+	{
+		// Note: Fires once per frame FTSTicker::AddTicker has a delay param if desired
+		HoldTickerHandle = FTSTicker::GetCoreTicker().AddTicker(FTickerDelegate::CreateUObject(this, &UCommonButtonBase::NativeOnHoldProgress));
+		if (HoldProgressRollbackTickerHandle.IsValid())
+		{
+			FTSTicker::GetCoreTicker().RemoveTicker(HoldProgressRollbackTickerHandle);
+			HoldProgressRollbackTickerHandle = nullptr;
+		}
+	}
+	if (TriggeringBindingHandle.IsValid())
+	{
+		TriggeringBindingHandle.ResetHold();
+	}
 }
 
 void UCommonButtonBase::HandleButtonReleased()
@@ -1223,6 +1386,23 @@ void UCommonButtonBase::HandleButtonReleased()
 		// Simulate hover events when using touch input
 		NativeOnUnhovered();
 	}
+
+	if (bRequiresHold && HoldTime > 0.f)
+	{
+		if (HoldRollbackTime <= UE_SMALL_NUMBER)
+		{
+			HoldReset();
+			NativeOnHoldProgress(0.f);
+		}
+		else
+		{
+			// Begin hold progress rollback
+			HoldProgressRollbackTickerHandle = FTSTicker::GetCoreTicker().AddTicker(FTickerDelegate::CreateUObject(this, &UCommonButtonBase::NativeOnHoldProgressRollback));
+            		
+			FTSTicker::GetCoreTicker().RemoveTicker(HoldTickerHandle);
+			HoldTickerHandle = nullptr;
+		}
+	}
 }
 
 FReply UCommonButtonBase::NativeOnFocusReceived(const FGeometry& InGeometry, const FFocusEvent& InFocusEvent)
diff --git a/Engine/Plugins/Experimental/CommonUI/Source/CommonUI/Private/CommonUIEditorSettings.cpp b/Engine/Plugins/Experimental/CommonUI/Source/CommonUI/Private/CommonUIEditorSettings.cpp
--- a/Engine/Plugins/Experimental/CommonUI/Source/CommonUI/Private/CommonUIEditorSettings.cpp
+++ b/Engine/Plugins/Experimental/CommonUI/Source/CommonUI/Private/CommonUIEditorSettings.cpp
@@ -31,6 +31,7 @@ void UCommonUIEditorSettings::LoadEditorData()
 		TemplateTextStyleClass = TemplateTextStyle.LoadSynchronous();
 		TemplateButtonStyleClass = TemplateButtonStyle.LoadSynchronous();
 		TemplateBorderStyleClass = TemplateBorderStyle.LoadSynchronous();
+		DefaultHoldDataClass = DefaultHoldData.LoadSynchronous();
 
 		if (GUObjectArray.IsDisregardForGC(this))
 		{
@@ -46,6 +47,10 @@ void UCommonUIEditorSettings::LoadEditorData()
 			{
 				TemplateBorderStyleClass->AddToRoot();
 			}
+			if (DefaultHoldDataClass)
+			{
+				DefaultHoldDataClass->AddToRoot();
+			}
 		}
 
 		bDefaultDataLoaded = true;
@@ -72,4 +77,10 @@ const TSubclassOf<UCommonBorderStyle>& UCommonUIEditorSettings::GetTemplateBorde
 
 	return TemplateBorderStyleClass;
 }
+
+const TSubclassOf<UCommonUIHoldData>& UCommonUIEditorSettings::GetDefaultHoldData() const
+{
+	ensure(bDefaultDataLoaded);
+	return DefaultHoldDataClass;
+}
 #endif
diff --git a/Engine/Plugins/Experimental/CommonUI/Source/CommonUI/Private/CommonUITypes.cpp b/Engine/Plugins/Experimental/CommonUI/Source/CommonUI/Private/CommonUITypes.cpp
--- a/Engine/Plugins/Experimental/CommonUI/Source/CommonUI/Private/CommonUITypes.cpp
+++ b/Engine/Plugins/Experimental/CommonUI/Source/CommonUI/Private/CommonUITypes.cpp
@@ -27,6 +27,7 @@ FCommonInputTypeInfo::FCommonInputTypeInfo()
 	OverrideBrush.DrawAs = ESlateBrushDrawType::NoDrawType;
 	bActionRequiresHold = false;
 	HoldTime = 0.5f;
+	HoldRollbackTime = 0.0f;
 }
 
 FKey FCommonInputTypeInfo::GetKey() const
diff --git a/Engine/Plugins/Experimental/CommonUI/Source/CommonUI/Private/Input/UIActionRouterTypes.cpp b/Engine/Plugins/Experimental/CommonUI/Source/CommonUI/Private/Input/UIActionRouterTypes.cpp
--- a/Engine/Plugins/Experimental/CommonUI/Source/CommonUI/Private/Input/UIActionRouterTypes.cpp
+++ b/Engine/Plugins/Experimental/CommonUI/Source/CommonUI/Private/Input/UIActionRouterTypes.cpp
@@ -17,6 +17,8 @@
 #include "Input/UIActionBinding.h"
 #include "Widgets/SViewport.h"
 #include "CommonButtonBase.h"
+#include "CommonUIEditorSettings.h"
+#include "ICommonUIModule.h"
 
 DEFINE_LOG_CATEGORY(LogUIActionRouter);
 
@@ -117,6 +119,7 @@ FUIActionBinding::FUIActionBinding(const UWidget& InBoundWidget, const FBindUIAc
 	, InputAction(BindArgs.InputAction)
 {
 	OnHoldActionProgressed.Add(BindArgs.OnHoldActionProgressed);
+	OnHoldActionPressed.Add(BindArgs.OnHoldActionPressed);
 
 	const auto RegisterKeyMappingFunc = 
 		[this](const FUIActionKeyMapping& KeyMapping)
@@ -164,22 +167,34 @@ FUIActionBinding::FUIActionBinding(const UWidget& InBoundWidget, const FBindUIAc
 			ActionDisplayName = LegacyActionData->DisplayName;
 		}
 
+		FName CurrentGamepadName = UCommonUIActionRouterBase::Get(InBoundWidget)->GetInputSubsystem().GetCurrentGamepadName();
+		
 		// KB/M
 		const FCommonInputTypeInfo& LegacyMapping_KBM = LegacyActionData->GetInputTypeInfo(ECommonInputType::MouseAndKeyboard, FCommonInputDefaults::GamepadGeneric);
-		FUIActionKeyMapping KeyMapping_KBM(LegacyMapping_KBM.GetKey(), LegacyMapping_KBM.bActionRequiresHold ? LegacyMapping_KBM.HoldTime : 0.f);
-		RegisterKeyMappingFunc(KeyMapping_KBM);
+		FUIActionKeyMapping KeyMapping_KBM(LegacyMapping_KBM.GetKey(), LegacyMapping_KBM.bActionRequiresHold ? LegacyMapping_KBM.HoldTime : 0.f, LegacyMapping_KBM.bActionRequiresHold ? LegacyMapping_KBM.HoldRollbackTime : 0.f);
 
 		// Touch
 		const FCommonInputTypeInfo& LegacyMapping_Touch = LegacyActionData->GetInputTypeInfo(ECommonInputType::Touch, FCommonInputDefaults::GamepadGeneric);
-		FUIActionKeyMapping KeyMapping_Touch(LegacyMapping_Touch.GetKey(), LegacyMapping_Touch.bActionRequiresHold ? LegacyMapping_Touch.HoldTime : 0.f);
-		RegisterKeyMappingFunc(KeyMapping_Touch);
+		FUIActionKeyMapping KeyMapping_Touch(LegacyMapping_Touch.GetKey(), LegacyMapping_Touch.bActionRequiresHold ? LegacyMapping_Touch.HoldTime : 0.f, LegacyMapping_Touch.bActionRequiresHold ? LegacyMapping_Touch.HoldRollbackTime : 0.f);
 
 		// Gamepad
 		// Note: This is definitely a wonky and roundabout way to get the gamepad type, but given that it's for legacy fixup we'll let it go until this can be deleted
 		UCommonUIActionRouterBase* ActionRouter = UCommonUIActionRouterBase::Get(InBoundWidget);
 		check(ActionRouter);
 		const FCommonInputTypeInfo& LegacyMapping_Gamepad = LegacyActionData->GetInputTypeInfo(ECommonInputType::Gamepad, ActionRouter->GetInputSubsystem().GetCurrentGamepadName());
-		FUIActionKeyMapping KeyMapping_Gamepad(LegacyMapping_Gamepad.GetKey(), LegacyMapping_Gamepad.bActionRequiresHold ? LegacyMapping_Gamepad.HoldTime : 0.f);
+		FUIActionKeyMapping KeyMapping_Gamepad(LegacyMapping_Gamepad.GetKey(), LegacyMapping_Gamepad.bActionRequiresHold ? LegacyMapping_Gamepad.HoldTime : 0.f, LegacyMapping_Gamepad.bActionRequiresHold ? LegacyMapping_Gamepad.HoldRollbackTime : 0.f);
+		
+		if (BindArgs.bForceHold && ICommonUIModule::GetEditorSettings().GetDefaultHoldData().GetDefaultObject())
+        {
+        	UCommonUIHoldData* CommonUIHoldBehaviorValues = ICommonUIModule::GetEditorSettings().GetDefaultHoldData().GetDefaultObject();
+        			
+        	KeyMapping_KBM.HoldTime = KeyMapping_KBM.HoldTime > 0.f ? KeyMapping_KBM.HoldTime : CommonUIHoldBehaviorValues->KeyboardAndMouse.HoldTime;
+        	KeyMapping_Touch.HoldTime = KeyMapping_Touch.HoldTime > 0.f ? KeyMapping_Touch.HoldTime : CommonUIHoldBehaviorValues->Touch.HoldTime;
+        	KeyMapping_Gamepad.HoldTime = KeyMapping_Gamepad.HoldTime > 0.f ? KeyMapping_Gamepad.HoldTime : CommonUIHoldBehaviorValues->Gamepad.HoldTime;
+        }
+        		
+        RegisterKeyMappingFunc(KeyMapping_KBM);
+        RegisterKeyMappingFunc(KeyMapping_Touch);
 		RegisterKeyMappingFunc(KeyMapping_Gamepad);
 	}
 
@@ -283,18 +298,27 @@ FString FUIActionBinding::ToDebugString() const
 
 void FUIActionBinding::BeginHold()
 {
+	if (HoldProgressRollbackTickerHandle.IsValid())
+	{
+		FTSTicker::GetCoreTicker().RemoveTicker(HoldProgressRollbackTickerHandle);
+		HoldProgressRollbackTickerHandle = nullptr;
+	}
+	
 	UE_LOG(LogUIActionRouter, VeryVerbose, TEXT("Hold pressed: %s"), *ActionName.ToString());
 
 	HoldStartTime = FPlatformTime::Seconds();
-	OnHoldActionProgressed.Broadcast(0.0f);
+	
+	const float HoldStartPercent = HoldTime > 0.f ? HoldStartSecond / HoldTime : 0.f;
+    OnHoldActionProgressed.Broadcast(HoldStartPercent);
 }
 
 bool FUIActionBinding::UpdateHold(float TargetHoldTime)
 {
-	const float SecondsHeld = GetSecondsHeld();
-	UE_LOG(LogUIActionRouter, VeryVerbose, TEXT("Hold repeating: %s for %f"), *ActionName.ToString(), SecondsHeld);
-
-	const float HeldPercent = FMath::Clamp(SecondsHeld / TargetHoldTime, 0.f, 1.f);
+	CurrentHoldSecond = FMath::Clamp(GetSecondsHeld() + HoldStartSecond, 0.f, TargetHoldTime);
+	UE_LOG(LogUIActionRouter, VeryVerbose, TEXT("Hold repeating: %s for %f"), *ActionName.ToString(), CurrentHoldSecond);
+	
+	const float HeldPercent = FMath::Clamp(CurrentHoldSecond / TargetHoldTime, 0.f, 1.f);
+	
 	OnHoldActionProgressed.Broadcast(HeldPercent);
 
 	if (HeldPercent >= 1.f)
@@ -302,8 +326,10 @@ bool FUIActionBinding::UpdateHold(float TargetHoldTime)
 		CancelHold();
 		OnHoldActionProgressed.Broadcast(0.0f); //approximate CommonButton::OnActionComplete without adding extra bind
 		OnExecuteAction.ExecuteIfBound();
+		HoldStartSecond = 0.0f;
+		CurrentHoldSecond = 0.0f;
 
-		UE_LOG(LogUIActionRouter, VeryVerbose, TEXT("Hold repeating: Fired completed action! %s after %f"), *ActionName.ToString(), SecondsHeld);
+		UE_LOG(LogUIActionRouter, VeryVerbose, TEXT("Hold repeating: Fired completed action! %s after %f"), *ActionName.ToString(), CurrentHoldSecond);
 		return true;
 	}
 	return false;
@@ -324,11 +350,52 @@ void FUIActionBinding::CancelHold()
 	HoldStartTime = -1.0;
 }
 
+void FUIActionBinding::BeginRollback(float TargetHoldRollbackTime, float InHoldTime, FUIActionBindingHandle BindingHandle)
+{
+	HoldRollbackMultiplier = InHoldTime / TargetHoldRollbackTime;
+	HoldTime = InHoldTime;
+	if (TSharedPtr<FUIActionBinding> InBinding = FindBinding(BindingHandle))
+	{
+		HoldProgressRollbackTickerHandle = FTSTicker::GetCoreTicker().AddTicker(FTickerDelegate::CreateLambda([InBinding](float DeltaTime)
+			{
+				const float HoldRollbackTime = InBinding->HoldTime / InBinding->HoldRollbackMultiplier;
+				InBinding->CurrentHoldSecond = FMath::Clamp(InBinding->CurrentHoldSecond - (DeltaTime * InBinding->HoldRollbackMultiplier), 0.f, HoldRollbackTime);
+				const float HoldRollbackPercent = FMath::Clamp(InBinding->CurrentHoldSecond / InBinding->HoldTime, 0.f, 1.f);
+				InBinding->OnHoldActionProgressed.Broadcast(HoldRollbackPercent);
+				InBinding->HoldStartSecond = InBinding->CurrentHoldSecond;
+			
+				if (HoldRollbackPercent <= 0.f)
+				{
+					FTSTicker::GetCoreTicker().RemoveTicker(InBinding->HoldProgressRollbackTickerHandle);
+					InBinding->HoldProgressRollbackTickerHandle = nullptr;
+					InBinding->HoldStartSecond = 0.f;
+					InBinding->CurrentHoldSecond = 0.f;
+						
+					return false;
+				}
+				
+				return true;
+			}));
+	}
+}
+
+
 bool FUIActionBinding::IsHoldActive() const
 {
 	return HoldStartTime >= 0.0;
 }
 
+void FUIActionBinding::ResetHold()
+{
+	if (HoldProgressRollbackTickerHandle.IsValid())
+	{
+		FTSTicker::GetCoreTicker().RemoveTicker(HoldProgressRollbackTickerHandle);
+		HoldProgressRollbackTickerHandle = nullptr;
+	}
+	HoldStartSecond = 0.0f;
+	CurrentHoldSecond = 0.0f;
+}
+
 //////////////////////////////////////////////////////////////////////////
 // FBindUIActionArgs
 //////////////////////////////////////////////////////////////////////////
@@ -388,6 +455,14 @@ void FUIActionBindingHandle::Unregister()
 	RegistrationId = INDEX_NONE;
 }
 
+void FUIActionBindingHandle::ResetHold()
+{
+	if (TSharedPtr<FUIActionBinding> Binding = FUIActionBinding::FindBinding(*this))
+	{
+		Binding->ResetHold();
+	}
+}
+
 FName FUIActionBindingHandle::GetActionName() const
 {
 	if (TSharedPtr<const FUIActionBinding> Binding = FUIActionBinding::FindBinding(*this))
@@ -499,6 +574,7 @@ EProcessHoldActionResult FActionRouterBindingCollection::ProcessHoldInput(ECommo
 						{
 							Binding->BeginHold();
 							ProcessResult = EProcessHoldActionResult::Handled;
+							Binding->OnHoldActionPressed.Broadcast();
 						}
 						else if (Binding->IsHoldActive())
 						{
@@ -514,8 +590,17 @@ EProcessHoldActionResult FActionRouterBindingCollection::ProcessHoldInput(ECommo
 								const float SecondsHeld = Binding->GetSecondsHeld();
 
 								UE_LOG(LogUIActionRouter, VeryVerbose, TEXT("Hold released: %s after %f"), *Binding->ActionName.ToString(), SecondsHeld);
-								Binding->CancelHold();
-								Binding->OnHoldActionProgressed.Broadcast(0.0f);
+
+								if (HoldMapping.HoldRollbackTime <= UE_SMALL_NUMBER)
+								{
+									Binding->OnHoldActionProgressed.Broadcast(0.f);
+									Binding->CancelHold();
+								}
+								else
+								{
+									Binding->CancelHold();
+									Binding->BeginRollback(HoldMapping.HoldRollbackTime, HoldMapping.HoldTime, Binding->Handle);
+								}
 
 								static const float PressToHoldThreshold = 0.25f;
 								if (SecondsHeld <= PressToHoldThreshold && SecondsHeld < HoldMapping.HoldTime)
diff --git a/Engine/Plugins/Experimental/CommonUI/Source/CommonUI/Public/CommonButtonBase.h b/Engine/Plugins/Experimental/CommonUI/Source/CommonUI/Public/CommonButtonBase.h
--- a/Engine/Plugins/Experimental/CommonUI/Source/CommonUI/Public/CommonButtonBase.h
+++ b/Engine/Plugins/Experimental/CommonUI/Source/CommonUI/Public/CommonButtonBase.h
@@ -9,6 +9,8 @@
 #include "FieldNotification/WidgetEventField.h"
 #include "Types/ISlateMetaData.h"
 #include "CommonInputModeTypes.h"
+#include "CommonInputBaseTypes.h"
+#include "Containers/Ticker.h"
 #include "CommonButtonBase.generated.h"
 
 class UCommonButtonBase;
@@ -507,6 +509,10 @@ protected:
 	/** Called via delegate when the input method changes */
 	UFUNCTION()
 	virtual void OnInputMethodChanged(ECommonInputType CurrentInputType);
+	
+	/** If HoldData is valid, assigns its values to Keyboard and Mouse, Gamepad and Touch, based off the Current Input Type. */
+    UFUNCTION()
+	virtual void UpdateHoldData(ECommonInputType CurrentInputType);
 
 	/** Associates this button at its priority with the given key */
 	virtual void BindTriggeringInputActionToClick();
@@ -634,15 +640,34 @@ protected:
 	UFUNCTION(BlueprintImplementableEvent, Category = "Common Button")
 	void OnActionProgress(float HeldPercent);
 	
+	/**
+    * By default, if bRequiresHold is true, the bound Input Action will be forced to require hold as well.
+    * The Bound Input action will use the row's hold values regardless of the row's bRequiresHold state.
+    */
+	UFUNCTION()
+	virtual bool GetConvertInputActionToHold();
+    
+    /** Bound to the hold progress of the bound key from the input action */
 	UFUNCTION()
 	virtual void NativeOnActionProgress(float HeldPercent);
 
+	/** Bound to the hold progress not related to the bound key */
+	UFUNCTION()
+	virtual bool NativeOnHoldProgress(float DeltaTime);
+	
+	/** Bound to the hold progress rollback not related to the bound key */
+	UFUNCTION()
+	virtual bool NativeOnHoldProgressRollback(float DeltaTime);
+	
 	/** Callback fired when hold events complete */
 	UFUNCTION(BlueprintImplementableEvent, Category = "Common Button")
 	void OnActionComplete();
 
 	UFUNCTION()
 	virtual void NativeOnActionComplete();
+	
+	UFUNCTION()
+    virtual void HoldReset();
 
 	virtual bool GetButtonAnalyticInfo(FString& ButtonName, FString& ABTestName, FString& ExtraData) const;
 
@@ -738,6 +763,14 @@ protected:
 	/** True if this button should use the default fallback input action (bool is useful for buttons that shouldn't because they are never directly hit via controller) */
 	UPROPERTY(EditAnywhere, BlueprintReadOnly, Category = Input, meta = (ExposeOnSpawn = true))
 	uint8 bShouldUseFallbackDefaultInputAction:1;
+	
+	/** True if this button should have a press and hold behavior, triggering the click when the specified hold time is met */
+	UPROPERTY(EditAnywhere, BlueprintReadOnly, Category = "Input|Hold", meta = (ExposeOnSpawn = true))
+	uint8 bRequiresHold:1;
+	
+	/** Press and Hold values used for Keyboard and Mouse, Gamepad and Touch, depending on the current input type */
+	UPROPERTY(EditAnywhere, BlueprintReadOnly, Category = "Input|Hold", meta = (EditCondition="bRequiresHold", ExposeOnSpawn = true))
+	TSubclassOf<UCommonUIHoldData> HoldData;
 
 private:
 
@@ -752,7 +785,7 @@ private:
 
 public:
 	/** The type of mouse action required by the user to trigger the button's 'Click' */
-	UPROPERTY(EditAnywhere, BlueprintReadOnly, Category = Input, meta = (ExposeOnSpawn = true))
+	UPROPERTY(EditAnywhere, BlueprintReadOnly, Category = Input, meta = (ExposeOnSFpawn = true))
 	TEnumAsByte<EButtonClickMethod::Type> ClickMethod;
 
 	UPROPERTY(EditAnywhere, BlueprintReadOnly, Category = Input)
@@ -815,6 +848,25 @@ protected:
 	FCommonButtonBaseClicked OnButtonBaseUnhovered;
 
 	FUIActionBindingHandle TriggeringBindingHandle;
+	
+	/** Press and hold time in seconds */
+	float HoldTime;
+	
+    /**
+    * Time (in seconds) for hold progress to go from 1.0 (completed) to 0.0.
+    * Used when the press and hold is interrupted.
+    * If set to 0, there will be no rollback and the hold progress will reset immediately.
+    */
+    float HoldRollbackTime;
+	
+    /** Current hold time for this button */
+    float CurrentHoldTime;
+	
+    /** Handle for ticker spawned for press and hold */
+    FTSTicker::FDelegateHandle HoldTickerHandle;
+	
+    /** Handle for ticker spawned for button hold rollback */
+    FTSTicker::FDelegateHandle HoldProgressRollbackTickerHandle;
 
 private:
 	friend class UCommonButtonGroupBase;
diff --git a/Engine/Plugins/Experimental/CommonUI/Source/CommonUI/Public/CommonUIEditorSettings.h b/Engine/Plugins/Experimental/CommonUI/Source/CommonUI/Public/CommonUIEditorSettings.h
--- a/Engine/Plugins/Experimental/CommonUI/Source/CommonUI/Public/CommonUIEditorSettings.h
+++ b/Engine/Plugins/Experimental/CommonUI/Source/CommonUI/Public/CommonUIEditorSettings.h
@@ -28,6 +28,9 @@ public:
 
 	const TSubclassOf<UCommonBorderStyle>& GetTemplateBorderStyle() const;
 
+	/** Default Hold Data */
+	const TSubclassOf<UCommonUIHoldData>& GetDefaultHoldData() const;
+
 private:
 	void LoadEditorData();
 #endif
@@ -45,6 +48,13 @@ private:
 	UPROPERTY(config, EditAnywhere, Category = "Border")
 	TSoftClassPtr<UCommonBorderStyle> TemplateBorderStyle;
 
+	/**
+	* Newly created CommonButton widgets will use these hold values by default if bRequiresHold is true.
+	* Inherits from UCommonUIHoldData.
+	*/
+	UPROPERTY(EditAnywhere, Config, Category = "Buttons|Hold")
+	TSoftClassPtr<UCommonUIHoldData> DefaultHoldData;
+
 #if WITH_EDITORONLY_DATA
 	UPROPERTY(Transient)
 	TSubclassOf<UCommonTextStyle> TemplateTextStyleClass;
@@ -54,6 +64,9 @@ private:
 
 	UPROPERTY(Transient)
 	TSubclassOf<UCommonBorderStyle> TemplateBorderStyleClass;
+
+	UPROPERTY(Transient)
+	TSubclassOf<UCommonUIHoldData> DefaultHoldDataClass;
 #endif
 private:
 	bool bDefaultDataLoaded;
diff --git a/Engine/Plugins/Experimental/CommonUI/Source/CommonUI/Public/CommonUITypes.h b/Engine/Plugins/Experimental/CommonUI/Source/CommonUI/Public/CommonUITypes.h
--- a/Engine/Plugins/Experimental/CommonUI/Source/CommonUI/Public/CommonUITypes.h
+++ b/Engine/Plugins/Experimental/CommonUI/Source/CommonUI/Public/CommonUITypes.h
@@ -64,8 +64,16 @@ public:
 	bool bActionRequiresHold;
 
 	/** The hold time in seconds */
-	UPROPERTY(EditAnywhere, Category = "CommonInput", meta = (EditCondition = "bActionRequiresHold"))
+	UPROPERTY(EditAnywhere, Category = "CommonInput", meta = (EditCondition = "bActionRequiresHold", ClampMin = "0.0", UIMin = "0.0"))
 	float HoldTime;
+
+	/**
+	*	Time (in seconds) for hold progress to go from 1.0 (completed) to 0.0.
+	*	If the hold interaction was interrupted, then hold progress starts to roll back decreasing its value.
+	*	Set to 0.0 to disable the rollback functionality.
+	*/
+	UPROPERTY(EditAnywhere, Category = "CommonInput", meta = (EditCondition = "bActionRequiresHold", ClampMin = "0", UIMin = "0", ClampMax = "10.0", UIMax = "10"))
+	float HoldRollbackTime;
 	
 	/** Override the brush specified by the Key Display Data  */
 	UPROPERTY(EditAnywhere, Category = "CommonInput")
@@ -76,7 +84,7 @@ public:
 		return Key == Other.Key &&
 			OverrrideState == Other.OverrrideState &&
 			bActionRequiresHold == Other.bActionRequiresHold &&
-			HoldTime == Other.HoldTime &&
+			HoldTime == Other.HoldTime &&HoldRollbackTime == Other.HoldRollbackTime &&
 			OverrideBrush == Other.OverrideBrush;
 	}
 
diff --git a/Engine/Plugins/Experimental/CommonUI/Source/CommonUI/Public/Input/CommonUIInputSettings.h b/Engine/Plugins/Experimental/CommonUI/Source/CommonUI/Public/Input/CommonUIInputSettings.h
--- a/Engine/Plugins/Experimental/CommonUI/Source/CommonUI/Public/Input/CommonUIInputSettings.h
+++ b/Engine/Plugins/Experimental/CommonUI/Source/CommonUI/Public/Input/CommonUIInputSettings.h
@@ -15,9 +15,10 @@ struct FUIActionKeyMapping
 
 public:
 	FUIActionKeyMapping() {}
-	FUIActionKeyMapping(FKey InKey, float InHoldTime)
+	FUIActionKeyMapping(FKey InKey, float InHoldTime, float InHoldRollbackTime)
 		: Key(InKey)
 		, HoldTime(InHoldTime)
+		, HoldRollbackTime(InHoldRollbackTime)
 	{}
 
 	/** A key that triggers this action */
@@ -27,6 +28,10 @@ public:
 	/** How long must the key be held down for the action to be executed? */
 	UPROPERTY(EditAnywhere, Config, Category = "UI Action Key Mapping")
 	float HoldTime = 0.f;
+	
+	/** Time (in seconds) for hold progress to go from 1.0 (completed) to 0.0. */
+	UPROPERTY(EditAnywhere, Config, Category = "UI Action Key Mapping")
+	float HoldRollbackTime = 0.f;
 
 	//@todo DanH: Is this actually wanted/needed at all? Do we really want different bindings for the same thing on Platform A vs Platform B, for instance?
 	///** 
diff --git a/Engine/Plugins/Experimental/CommonUI/Source/CommonUI/Public/Input/CommonUIInputTypes.h b/Engine/Plugins/Experimental/CommonUI/Source/CommonUI/Public/Input/CommonUIInputTypes.h
--- a/Engine/Plugins/Experimental/CommonUI/Source/CommonUI/Public/Input/CommonUIInputTypes.h
+++ b/Engine/Plugins/Experimental/CommonUI/Source/CommonUI/Public/Input/CommonUIInputTypes.h
@@ -75,6 +75,9 @@ struct COMMONUI_API FBindUIActionArgs
 	/** Whether this binding can/should be displayed in a CommonActionBar (if one exists) */
 	bool bDisplayInActionBar = true;
 
+	/** True implies we will add default hold times if the current action is not a hold action */
+    bool bForceHold = false;
+
 	/** Optional display name to associate with this binding instead of the default */
 	FText OverrideDisplayName;
 
@@ -83,6 +86,10 @@ struct COMMONUI_API FBindUIActionArgs
 	/** If the bound action has any hold mappings, this will fire each frame while held. Has no bearing on actual execution and wholly irrelevant for non-hold actions */
 	DECLARE_DELEGATE_OneParam(FOnHoldActionProgressed, float);
 	FOnHoldActionProgressed OnHoldActionProgressed;
+	
+	/** If the bound action has any hold mappings, this will fire when the hold begins. Has no bearing on actual execution and wholly irrelevant for non-hold actions */
+	DECLARE_DELEGATE(FOnHoldActionPressed);
+	FOnHoldActionPressed OnHoldActionPressed;
 };
 
 #if UE_ENABLE_INCLUDE_ORDER_DEPRECATED_IN_5_2
diff --git a/Engine/Plugins/Experimental/CommonUI/Source/CommonUI/Public/Input/UIActionBinding.h b/Engine/Plugins/Experimental/CommonUI/Source/CommonUI/Public/Input/UIActionBinding.h
--- a/Engine/Plugins/Experimental/CommonUI/Source/CommonUI/Public/Input/UIActionBinding.h
+++ b/Engine/Plugins/Experimental/CommonUI/Source/CommonUI/Public/Input/UIActionBinding.h
@@ -45,8 +45,10 @@ struct COMMONUI_API FUIActionBinding
 	void BeginHold();
 	bool UpdateHold(float TargetHoldTime);
 	void CancelHold();
+	void BeginRollback(float TargetHoldRollbackTime, float HoldTime, FUIActionBindingHandle BindingHandle);
 	double GetSecondsHeld() const;
 	bool IsHoldActive() const;
+	void ResetHold();
 
 	FName ActionName;
 	EInputEvent InputEvent;
@@ -69,6 +71,9 @@ struct COMMONUI_API FUIActionBinding
 	DECLARE_MULTICAST_DELEGATE_OneParam(FOnHoldActionProgressedMulticast, float);
 	FOnHoldActionProgressedMulticast OnHoldActionProgressed;
 
+	DECLARE_MULTICAST_DELEGATE(FOnHoldActionPressed);
+	FOnHoldActionPressed OnHoldActionPressed;
+
 	// @TODO: Rename non-legacy in 5.3. We no longer have any active plans to remove data tables in CommonUI.
 	FDataTableRowHandle LegacyActionTableRow;
 
@@ -77,7 +82,23 @@ struct COMMONUI_API FUIActionBinding
 private:
 	FUIActionBinding(const UWidget& InBoundWidget, const FBindUIActionArgs& BindArgs);
 	
+	// At what time in seconds did the hold start?
 	double HoldStartTime = -1.0;
+    	
+	// At what second will the hold start?
+	double HoldStartSecond = 0.0;
+	
+	// At what second is the hold progress at?
+	double CurrentHoldSecond = 0.0;
+	
+	// Multiplier for the time (in seconds) for hold progress to go from 1.0 (completed) to 0.0.
+    double HoldRollbackMultiplier = 1.0;
+	
+	// Target time (in seconds) for the hold progress to go from 0.0 to 1.0 (completed).
+	double HoldTime = 0.0;
+    	
+	// Handle for ticker spawned for button hold rollback
+	FTSTicker::FDelegateHandle HoldProgressRollbackTickerHandle;
 
 	static int32 IdCounter;
 	static TMap<FUIActionBindingHandle, TSharedPtr<FUIActionBinding>> AllRegistrationsByHandle;
diff --git a/Engine/Plugins/Experimental/CommonUI/Source/CommonUI/Public/Input/UIActionBindingHandle.h b/Engine/Plugins/Experimental/CommonUI/Source/CommonUI/Public/Input/UIActionBindingHandle.h
--- a/Engine/Plugins/Experimental/CommonUI/Source/CommonUI/Public/Input/UIActionBindingHandle.h
+++ b/Engine/Plugins/Experimental/CommonUI/Source/CommonUI/Public/Input/UIActionBindingHandle.h
@@ -17,6 +17,9 @@ public:
 	bool IsValid() const;
 	void Unregister();
 
+	/** Calls ResetHold() on the Action Binding, which in turn resets the hold progress to 0.0  */
+	void ResetHold();
+
 	FName GetActionName() const;
 
 	FText GetDisplayName() const;
