- Modified Hold Data settings be used from the CommonInputSettings instead of the CommonUIEditorSettings and avoid compilation errors.

diff --git a/Engine/Plugins/Experimental/CommonUI/Source/CommonInput/Private/CommonInputSettings.cpp b/Engine/Plugins/Experimental/CommonUI/Source/CommonInput/Private/CommonInputSettings.cpp
--- a/Engine/Plugins/Experimental/CommonUI/Source/CommonInput/Private/CommonInputSettings.cpp
+++ b/Engine/Plugins/Experimental/CommonUI/Source/CommonInput/Private/CommonInputSettings.cpp
@@ -135,6 +135,21 @@ FDataTableRowHandle UCommonInputSettings::GetDefaultBackAction() const
 	return FDataTableRowHandle();
 }
 
+TSubclassOf<UCommonUIHoldData> UCommonInputSettings::GetDefaultHoldData() const
+{
+	ensure(bInputDataLoaded);
+
+		if (InputDataClass)
+		{
+			if (const UCommonUIInputData* InputDataPtr = InputDataClass.GetDefaultObject())
+			{
+				TSubclassOf<UCommonUIHoldData> DefaultHoldDataLoaded = InputDataPtr->DefaultHoldData.LoadSynchronous();
+				return DefaultHoldDataLoaded;
+			}
+		}
+		return nullptr;
+}
+
 UInputAction* UCommonInputSettings::GetEnhancedInputClickAction() const
 {
 	ensureMsgf(bInputDataLoaded, TEXT("Warning, CommonUI input data not loaded."));
diff --git a/Engine/Plugins/Experimental/CommonUI/Source/CommonInput/Public/CommonInputBaseTypes.h b/Engine/Plugins/Experimental/CommonUI/Source/CommonInput/Public/CommonInputBaseTypes.h
--- a/Engine/Plugins/Experimental/CommonUI/Source/CommonInput/Public/CommonInputBaseTypes.h
+++ b/Engine/Plugins/Experimental/CommonUI/Source/CommonInput/Public/CommonInputBaseTypes.h
@@ -19,6 +19,7 @@
 #include "CommonInputBaseTypes.generated.h"
 
 
+class UCommonUIHoldData;
 class UTexture2D;
 class UMaterial;
 class UCommonInputSettings; 
@@ -104,6 +105,13 @@ public:
 	UPROPERTY(EditDefaultsOnly, Category = "Properties", meta = (RowType = "/Script/CommonUI.CommonInputActionDataBase"))
 	FDataTableRowHandle DefaultBackAction;
 
+	/**
+    * Newly created CommonButton widgets will use these hold values by default if bRequiresHold is true.
+    * Inherits from UCommonUIHoldData.
+    */
+    UPROPERTY(EditDefaultsOnly, Category = "Properties")
+    TSoftClassPtr<UCommonUIHoldData> DefaultHoldData;
+
 	UPROPERTY(EditDefaultsOnly, Category = "Properties", meta = (EditCondition = "CommonInput.CommonInputSettings.IsEnhancedInputSupportEnabled", EditConditionHides))
 	TObjectPtr<UInputAction> EnhancedInputClickAction;
 
diff --git a/Engine/Plugins/Experimental/CommonUI/Source/CommonInput/Public/CommonInputSettings.h b/Engine/Plugins/Experimental/CommonUI/Source/CommonInput/Public/CommonInputSettings.h
--- a/Engine/Plugins/Experimental/CommonUI/Source/CommonInput/Public/CommonInputSettings.h
+++ b/Engine/Plugins/Experimental/CommonUI/Source/CommonInput/Public/CommonInputSettings.h
@@ -39,6 +39,9 @@ public:
 
 	FDataTableRowHandle GetDefaultClickAction() const;
 	FDataTableRowHandle GetDefaultBackAction() const;
+	
+	/** Default Hold Data */
+	TSubclassOf<UCommonUIHoldData> GetDefaultHoldData() const;
 
 	UInputAction* GetEnhancedInputClickAction() const;
 	UInputAction* GetEnhancedInputBackAction() const;
diff --git a/Engine/Plugins/Experimental/CommonUI/Source/CommonUI/Private/CommonButtonBase.cpp b/Engine/Plugins/Experimental/CommonUI/Source/CommonUI/Private/CommonButtonBase.cpp
--- a/Engine/Plugins/Experimental/CommonUI/Source/CommonUI/Private/CommonButtonBase.cpp
+++ b/Engine/Plugins/Experimental/CommonUI/Source/CommonUI/Private/CommonButtonBase.cpp
@@ -375,18 +375,14 @@ void UCommonButtonBase::PostLoad()
 	{
 		UCommonUIEditorSettings& Settings = ICommonUIModule::GetEditorSettings();
 		Settings.ConditionalPostLoad();
-		
-		if (!Style && !bStyleNoLongerNeedsConversion)
-		{
-			Style = Settings.GetTemplateButtonStyle();
-		}
-		if (!HoldData)
-		{
-			HoldData = Settings.GetDefaultHoldData();
-		}
+		Style = Settings.GetTemplateButtonStyle();
 	}
-
 	bStyleNoLongerNeedsConversion = true;
+
+	if (!HoldData && ICommonInputModule::GetSettings().GetDefaultHoldData())
+	{
+		HoldData = ICommonInputModule::GetSettings().GetDefaultHoldData();
+	}
 #endif
 }
 
@@ -406,9 +402,9 @@ void UCommonButtonBase::OnCreationFromPalette()
 	{
 		Style = ICommonUIModule::GetEditorSettings().GetTemplateButtonStyle();
 	}
-	if (!HoldData)
+	if (!HoldData && ICommonInputModule::GetSettings().GetDefaultHoldData())
 	{
-		HoldData = ICommonUIModule::GetEditorSettings().GetDefaultHoldData();
+		HoldData = ICommonInputModule::GetSettings().GetDefaultHoldData();
 	}
 	Super::OnCreationFromPalette();
 }
@@ -462,9 +458,9 @@ UCommonButtonInternalBase* UCommonButtonBase::ConstructInternalButton()
 
 void UCommonButtonBase::NativeConstruct()
 {
-	if (bRequiresHold && !HoldData)
+	if (!HoldData && ICommonInputModule::GetSettings().GetDefaultHoldData())
 	{
-		HoldData = ICommonUIModule::GetEditorSettings().GetDefaultHoldData();
+		HoldData = ICommonInputModule::GetSettings().GetDefaultHoldData();
 	}
 	
 	BindTriggeringInputActionToClick();
diff --git a/Engine/Plugins/Experimental/CommonUI/Source/CommonUI/Private/CommonUIEditorSettings.cpp b/Engine/Plugins/Experimental/CommonUI/Source/CommonUI/Private/CommonUIEditorSettings.cpp
--- a/Engine/Plugins/Experimental/CommonUI/Source/CommonUI/Private/CommonUIEditorSettings.cpp
+++ b/Engine/Plugins/Experimental/CommonUI/Source/CommonUI/Private/CommonUIEditorSettings.cpp
@@ -31,7 +31,6 @@ void UCommonUIEditorSettings::LoadEditorData()
 		TemplateTextStyleClass = TemplateTextStyle.LoadSynchronous();
 		TemplateButtonStyleClass = TemplateButtonStyle.LoadSynchronous();
 		TemplateBorderStyleClass = TemplateBorderStyle.LoadSynchronous();
-		DefaultHoldDataClass = DefaultHoldData.LoadSynchronous();
 
 		if (GUObjectArray.IsDisregardForGC(this))
 		{
@@ -47,10 +46,6 @@ void UCommonUIEditorSettings::LoadEditorData()
 			{
 				TemplateBorderStyleClass->AddToRoot();
 			}
-			if (DefaultHoldDataClass)
-			{
-				DefaultHoldDataClass->AddToRoot();
-			}
 		}
 
 		bDefaultDataLoaded = true;
@@ -77,10 +72,4 @@ const TSubclassOf<UCommonBorderStyle>& UCommonUIEditorSettings::GetTemplateBorde
 
 	return TemplateBorderStyleClass;
 }
-
-const TSubclassOf<UCommonUIHoldData>& UCommonUIEditorSettings::GetDefaultHoldData() const
-{
-	ensure(bDefaultDataLoaded);
-	return DefaultHoldDataClass;
-}
 #endif
diff --git a/Engine/Plugins/Experimental/CommonUI/Source/CommonUI/Private/Input/UIActionRouterTypes.cpp b/Engine/Plugins/Experimental/CommonUI/Source/CommonUI/Private/Input/UIActionRouterTypes.cpp
--- a/Engine/Plugins/Experimental/CommonUI/Source/CommonUI/Private/Input/UIActionRouterTypes.cpp
+++ b/Engine/Plugins/Experimental/CommonUI/Source/CommonUI/Private/Input/UIActionRouterTypes.cpp
@@ -184,9 +184,9 @@ FUIActionBinding::FUIActionBinding(const UWidget& InBoundWidget, const FBindUIAc
 		const FCommonInputTypeInfo& LegacyMapping_Gamepad = LegacyActionData->GetInputTypeInfo(ECommonInputType::Gamepad, ActionRouter->GetInputSubsystem().GetCurrentGamepadName());
 		FUIActionKeyMapping KeyMapping_Gamepad(LegacyMapping_Gamepad.GetKey(), LegacyMapping_Gamepad.bActionRequiresHold ? LegacyMapping_Gamepad.HoldTime : 0.f, LegacyMapping_Gamepad.bActionRequiresHold ? LegacyMapping_Gamepad.HoldRollbackTime : 0.f);
 		
-		if (BindArgs.bForceHold && ICommonUIModule::GetEditorSettings().GetDefaultHoldData().GetDefaultObject())
+		if (BindArgs.bForceHold && ICommonInputModule::GetSettings().GetDefaultHoldData())
         {
-        	UCommonUIHoldData* CommonUIHoldBehaviorValues = ICommonUIModule::GetEditorSettings().GetDefaultHoldData().GetDefaultObject();
+        	UCommonUIHoldData* CommonUIHoldBehaviorValues = ICommonInputModule::GetSettings().GetDefaultHoldData().GetDefaultObject();
         			
         	KeyMapping_KBM.HoldTime = KeyMapping_KBM.HoldTime > 0.f ? KeyMapping_KBM.HoldTime : CommonUIHoldBehaviorValues->KeyboardAndMouse.HoldTime;
         	KeyMapping_Touch.HoldTime = KeyMapping_Touch.HoldTime > 0.f ? KeyMapping_Touch.HoldTime : CommonUIHoldBehaviorValues->Touch.HoldTime;
diff --git a/Engine/Plugins/Experimental/CommonUI/Source/CommonUI/Public/CommonUIEditorSettings.h b/Engine/Plugins/Experimental/CommonUI/Source/CommonUI/Public/CommonUIEditorSettings.h
--- a/Engine/Plugins/Experimental/CommonUI/Source/CommonUI/Public/CommonUIEditorSettings.h
+++ b/Engine/Plugins/Experimental/CommonUI/Source/CommonUI/Public/CommonUIEditorSettings.h
@@ -28,9 +28,6 @@ public:
 
 	const TSubclassOf<UCommonBorderStyle>& GetTemplateBorderStyle() const;
 
-	/** Default Hold Data */
-	const TSubclassOf<UCommonUIHoldData>& GetDefaultHoldData() const;
-
 private:
 	void LoadEditorData();
 #endif
@@ -48,13 +45,6 @@ private:
 	UPROPERTY(config, EditAnywhere, Category = "Border")
 	TSoftClassPtr<UCommonBorderStyle> TemplateBorderStyle;
 
-	/**
-	* Newly created CommonButton widgets will use these hold values by default if bRequiresHold is true.
-	* Inherits from UCommonUIHoldData.
-	*/
-	UPROPERTY(EditAnywhere, Config, Category = "Buttons|Hold")
-	TSoftClassPtr<UCommonUIHoldData> DefaultHoldData;
-
 #if WITH_EDITORONLY_DATA
 	UPROPERTY(Transient)
 	TSubclassOf<UCommonTextStyle> TemplateTextStyleClass;
@@ -64,9 +54,6 @@ private:
 
 	UPROPERTY(Transient)
 	TSubclassOf<UCommonBorderStyle> TemplateBorderStyleClass;
-
-	UPROPERTY(Transient)
-	TSubclassOf<UCommonUIHoldData> DefaultHoldDataClass;
 #endif
 private:
 	bool bDefaultDataLoaded;
