- Remove Hold Data assigning from PostLoad() in CommonButtonBase. Instead call it in NativeConstruct().

diff --git a/Engine/Plugins/Experimental/CommonUI/Source/CommonUI/Private/CommonButtonBase.cpp b/Engine/Plugins/Experimental/CommonUI/Source/CommonUI/Private/CommonButtonBase.cpp
--- a/Engine/Plugins/Experimental/CommonUI/Source/CommonUI/Private/CommonButtonBase.cpp
+++ b/Engine/Plugins/Experimental/CommonUI/Source/CommonUI/Private/CommonButtonBase.cpp
@@ -378,11 +378,6 @@ void UCommonButtonBase::PostLoad()
 		Style = Settings.GetTemplateButtonStyle();
 	}
 	bStyleNoLongerNeedsConversion = true;
-
-	if (!HoldData && ICommonInputModule::GetSettings().GetDefaultHoldData())
-	{
-		HoldData = ICommonInputModule::GetSettings().GetDefaultHoldData();
-	}
 #endif
 }
 
