Initializing floats within FInputHoldData struct.

diff --git a/Engine/Plugins/Experimental/CommonUI/Source/CommonInput/Public/CommonInputBaseTypes.h b/Engine/Plugins/Experimental/CommonUI/Source/CommonInput/Public/CommonInputBaseTypes.h
--- a/Engine/Plugins/Experimental/CommonUI/Source/CommonInput/Public/CommonInputBaseTypes.h
+++ b/Engine/Plugins/Experimental/CommonUI/Source/CommonInput/Public/CommonInputBaseTypes.h
@@ -84,9 +84,10 @@ struct FInputHoldData
 {
 	GENERATED_BODY()
 	UPROPERTY(EditDefaultsOnly, Category = "Properties", meta = (ClampMin = "0.0", UIMin = "0.0"))
-	float HoldTime;
+	float HoldTime = 0.0;
+	
 	UPROPERTY(EditDefaultsOnly, Category = "Properties", meta = (ClampMin = "0.0", UIMin = "0.0"))
-	float HoldRollbackTime;
+	float HoldRollbackTime = 0.0;
 };
 
 /* Derive from this class to store the Input data. It is referenced in the Common Input Settings, found in the project settings UI. */
