Modified input binding rollback functionality to avoid extending ptr lifetime.

diff --git a/Engine/Plugins/Experimental/CommonUI/Source/CommonUI/Private/Input/UIActionRouterTypes.cpp b/Engine/Plugins/Experimental/CommonUI/Source/CommonUI/Private/Input/UIActionRouterTypes.cpp
--- a/Engine/Plugins/Experimental/CommonUI/Source/CommonUI/Private/Input/UIActionRouterTypes.cpp
+++ b/Engine/Plugins/Experimental/CommonUI/Source/CommonUI/Private/Input/UIActionRouterTypes.cpp
@@ -354,32 +354,34 @@ void FUIActionBinding::BeginRollback(float TargetHoldRollbackTime, float InHoldT
 {
 	HoldRollbackMultiplier = InHoldTime / TargetHoldRollbackTime;
 	HoldTime = InHoldTime;
-	if (TSharedPtr<FUIActionBinding> InBinding = FindBinding(BindingHandle))
+	
+	HoldProgressRollbackTickerHandle = FTSTicker::GetCoreTicker().AddTicker(FTickerDelegate::CreateLambda([BindingHandle](float DeltaTime)
 	{
-		HoldProgressRollbackTickerHandle = FTSTicker::GetCoreTicker().AddTicker(FTickerDelegate::CreateLambda([InBinding](float DeltaTime)
+		if (TSharedPtr<FUIActionBinding> InBinding = FindBinding(BindingHandle))
+		{
+			const float HoldRollbackTime = InBinding->HoldTime / InBinding->HoldRollbackMultiplier;
+			InBinding->CurrentHoldSecond = FMath::Clamp(InBinding->CurrentHoldSecond - (DeltaTime * InBinding->HoldRollbackMultiplier), 0.f, HoldRollbackTime);
+			const float HoldRollbackPercent = FMath::Clamp(InBinding->CurrentHoldSecond / InBinding->HoldTime, 0.f, 1.f);
+			InBinding->OnHoldActionProgressed.Broadcast(HoldRollbackPercent);
+			InBinding->HoldStartSecond = InBinding->CurrentHoldSecond;
+            			
+			if (HoldRollbackPercent <= 0.f)
 			{
-				const float HoldRollbackTime = InBinding->HoldTime / InBinding->HoldRollbackMultiplier;
-				InBinding->CurrentHoldSecond = FMath::Clamp(InBinding->CurrentHoldSecond - (DeltaTime * InBinding->HoldRollbackMultiplier), 0.f, HoldRollbackTime);
-				const float HoldRollbackPercent = FMath::Clamp(InBinding->CurrentHoldSecond / InBinding->HoldTime, 0.f, 1.f);
-				InBinding->OnHoldActionProgressed.Broadcast(HoldRollbackPercent);
-				InBinding->HoldStartSecond = InBinding->CurrentHoldSecond;
-			
-				if (HoldRollbackPercent <= 0.f)
-				{
-					FTSTicker::GetCoreTicker().RemoveTicker(InBinding->HoldProgressRollbackTickerHandle);
-					InBinding->HoldProgressRollbackTickerHandle = nullptr;
-					InBinding->HoldStartSecond = 0.f;
-					InBinding->CurrentHoldSecond = 0.f;
-						
-					return false;
-				}
-				
-				return true;
-			}));
-	}
+				FTSTicker::GetCoreTicker().RemoveTicker(InBinding->HoldProgressRollbackTickerHandle);
+				InBinding->HoldProgressRollbackTickerHandle = nullptr;
+				InBinding->HoldStartSecond = 0.f;
+				InBinding->CurrentHoldSecond = 0.f;
+            						
+				return false;
+			}
+            				
+			return true;
+		}
+		
+		return false;			
+	}));
 }
 
-
 bool FUIActionBinding::IsHoldActive() const
 {
 	return HoldStartTime >= 0.0;
