Standardize usage of west const from code introduced in CL 24097276

diff --git a/Engine/Plugins/Experimental/CommonUI/Source/CommonUI/Public/Input/CommonBoundActionBar.h b/Engine/Plugins/Experimental/CommonUI/Source/CommonUI/Public/Input/CommonBoundActionBar.h
--- a/Engine/Plugins/Experimental/CommonUI/Source/CommonUI/Public/Input/CommonBoundActionBar.h
+++ b/Engine/Plugins/Experimental/CommonUI/Source/CommonUI/Public/Input/CommonBoundActionBar.h
@@ -37,7 +37,7 @@ protected:
 	virtual void SynchronizeProperties() override;
 	virtual void ReleaseSlateResources(bool bReleaseChildren) override;
 
-	virtual void NativeOnActionButtonCreated(ICommonBoundActionButtonInterface* ActionButton, FUIActionBindingHandle const& RepresentedAction) { }
+	virtual void NativeOnActionButtonCreated(ICommonBoundActionButtonInterface* ActionButton, const FUIActionBindingHandle& RepresentedAction) { }
 
 #if WITH_EDITOR
 	void ValidateCompiledDefaults(IWidgetCompilerLog& CompileLog) const override;
