Remove Common Input Settings Hold Data ensure bInputDataLoaded for Hold Data

diff --git a/Engine/Plugins/Experimental/CommonUI/Source/CommonInput/Private/CommonInputSettings.cpp b/Engine/Plugins/Experimental/CommonUI/Source/CommonInput/Private/CommonInputSettings.cpp
--- a/Engine/Plugins/Experimental/CommonUI/Source/CommonInput/Private/CommonInputSettings.cpp
+++ b/Engine/Plugins/Experimental/CommonUI/Source/CommonInput/Private/CommonInputSettings.cpp
@@ -137,17 +137,15 @@ FDataTableRowHandle UCommonInputSettings::GetDefaultBackAction() const
 
 TSubclassOf<UCommonUIHoldData> UCommonInputSettings::GetDefaultHoldData() const
 {
-	ensure(bInputDataLoaded);
-
-		if (InputDataClass)
+	if (InputDataClass)
+	{
+		if (const UCommonUIInputData* InputDataPtr = InputDataClass.GetDefaultObject())
 		{
-			if (const UCommonUIInputData* InputDataPtr = InputDataClass.GetDefaultObject())
-			{
-				TSubclassOf<UCommonUIHoldData> DefaultHoldDataLoaded = InputDataPtr->DefaultHoldData.LoadSynchronous();
-				return DefaultHoldDataLoaded;
-			}
+			TSubclassOf<UCommonUIHoldData> DefaultHoldDataLoaded = InputDataPtr->DefaultHoldData.LoadSynchronous();
+			return DefaultHoldDataLoaded;
 		}
-		return nullptr;
+	}
+	return nullptr;
 }
 
 UInputAction* UCommonInputSettings::GetEnhancedInputClickAction() const
