Manually merge 25069522

diff --git a/Engine/Plugins/Experimental/CommonUI/Source/CommonInput/Private/CommonInputSettings.cpp b/Engine/Plugins/Experimental/CommonUI/Source/CommonInput/Private/CommonInputSettings.cpp
--- a/Engine/Plugins/Experimental/CommonUI/Source/CommonInput/Private/CommonInputSettings.cpp
+++ b/Engine/Plugins/Experimental/CommonUI/Source/CommonInput/Private/CommonInputSettings.cpp
@@ -123,8 +123,6 @@ FDataTableRowHandle UCommonInputSettings::GetDefaultClickAction() const
 
 FDataTableRowHandle UCommonInputSettings::GetDefaultBackAction() const
 {
-	ensure(bInputDataLoaded);
-
 	if (InputDataClass)
 	{
 		if (const UCommonUIInputData* InputDataPtr = InputDataClass.GetDefaultObject())
