Rename FUICameraConfig to FActivationMetadata

diff --git a/Engine/Plugins/Experimental/CommonUI/Source/CommonUI/Private/CommonActivatableWidget.cpp b/Engine/Plugins/Experimental/CommonUI/Source/CommonUI/Private/CommonActivatableWidget.cpp
--- a/Engine/Plugins/Experimental/CommonUI/Source/CommonUI/Private/CommonActivatableWidget.cpp
+++ b/Engine/Plugins/Experimental/CommonUI/Source/CommonUI/Private/CommonActivatableWidget.cpp
@@ -67,9 +67,9 @@ UWidget* UCommonActivatableWidget::GetDesiredFocusTarget() const
 	return NativeGetDesiredFocusTarget();
 }
 
-TOptional<FUICameraConfig> UCommonActivatableWidget::GetDesiredCameraConfig() const
+TOptional<FActivationMetadata> UCommonActivatableWidget::GetActivationMetadata() const
 {
-	return TOptional<FUICameraConfig>();
+	return TOptional<FActivationMetadata>();
 }
 
 UWidget* UCommonActivatableWidget::NativeGetDesiredFocusTarget() const
diff --git a/Engine/Plugins/Experimental/CommonUI/Source/CommonUI/Private/Input/CommonUIActionRouterBase.cpp b/Engine/Plugins/Experimental/CommonUI/Source/CommonUI/Private/Input/CommonUIActionRouterBase.cpp
--- a/Engine/Plugins/Experimental/CommonUI/Source/CommonUI/Private/Input/CommonUIActionRouterBase.cpp
+++ b/Engine/Plugins/Experimental/CommonUI/Source/CommonUI/Private/Input/CommonUIActionRouterBase.cpp
@@ -1445,9 +1445,9 @@ void UCommonUIActionRouterBase::ApplyUIInputConfig(const FUIInputConfig& NewConf
 	}
 }
 
-void UCommonUIActionRouterBase::SetActiveUICameraConfig(const FUICameraConfig& NewConfig)
+void UCommonUIActionRouterBase::SetActiveActivationMetadata(const FActivationMetadata& NewConfig)
 {
-	OnCameraConfigChanged().Broadcast(NewConfig);
+	OnActivationMetadataChanged().Broadcast(NewConfig);
 }
 
 void UCommonUIActionRouterBase::HandleActivatableWidgetRebuilding(UCommonActivatableWidget& RebuildingWidget)
diff --git a/Engine/Plugins/Experimental/CommonUI/Source/CommonUI/Private/Input/UIActionRouterTypes.cpp b/Engine/Plugins/Experimental/CommonUI/Source/CommonUI/Private/Input/UIActionRouterTypes.cpp
--- a/Engine/Plugins/Experimental/CommonUI/Source/CommonUI/Private/Input/UIActionRouterTypes.cpp
+++ b/Engine/Plugins/Experimental/CommonUI/Source/CommonUI/Private/Input/UIActionRouterTypes.cpp
@@ -968,12 +968,12 @@ TOptional<FUIInputConfig> FActivatableTreeNode::FindDesiredActionDomainInputConf
 	return  bHasActionDomainConfig ? FUIInputConfig(ActionDomain->InputMode, ActionDomain->MouseCaptureMode) : TOptional<FUIInputConfig>();
 }
 
-FUICameraConfig FActivatableTreeNode::FindDesiredCameraConfig() const
+FActivationMetadata FActivatableTreeNode::FindActivationMetadata() const
 {
-	TOptional<FUICameraConfig> DesiredConfig = ensure(RepresentedWidget.IsValid()) ? RepresentedWidget->GetDesiredCameraConfig() : TOptional<FUICameraConfig>();
+	TOptional<FActivationMetadata> DesiredConfig = ensure(RepresentedWidget.IsValid()) ? RepresentedWidget->GetActivationMetadata() : TOptional<FActivationMetadata>();
 	if (!DesiredConfig.IsSet() && Parent.IsValid())
 	{
-		DesiredConfig = Parent.Pin()->FindDesiredCameraConfig();
+		DesiredConfig = Parent.Pin()->FindActivationMetadata();
 	}
 
 	if (DesiredConfig.IsSet())
@@ -981,7 +981,7 @@ FUICameraConfig FActivatableTreeNode::FindDesiredCameraConfig() const
 		return DesiredConfig.GetValue();
 	}
 
-	return FUICameraConfig();
+	return FActivationMetadata();
 }
 
 void FActivatableTreeNode::AddScrollRecipient(const UWidget& ScrollRecipient)
@@ -1388,7 +1388,7 @@ void FActivatableTreeRoot::ApplyLeafmostNodeConfig()
 #endif
 	if (FActivatableTreeNodePtr PinnedLeafmostNode = LeafmostActiveNode.Pin())
 	{
-		GetActionRouter().SetActiveUICameraConfig(PinnedLeafmostNode->FindDesiredCameraConfig());
+		GetActionRouter().SetActiveActivationMetadata(PinnedLeafmostNode->FindActivationMetadata());
 
 		if (ensure(PinnedLeafmostNode->IsReceivingInput()))
 		{
diff --git a/Engine/Plugins/Experimental/CommonUI/Source/CommonUI/Private/Input/UIActionRouterTypes.h b/Engine/Plugins/Experimental/CommonUI/Source/CommonUI/Private/Input/UIActionRouterTypes.h
--- a/Engine/Plugins/Experimental/CommonUI/Source/CommonUI/Private/Input/UIActionRouterTypes.h
+++ b/Engine/Plugins/Experimental/CommonUI/Source/CommonUI/Private/Input/UIActionRouterTypes.h
@@ -121,7 +121,7 @@ public:
 	int32 GetLastPaintLayer() const;
 	TOptional<FUIInputConfig> FindDesiredInputConfig() const;
 	TOptional<FUIInputConfig> FindDesiredActionDomainInputConfig() const;
-	FUICameraConfig FindDesiredCameraConfig() const;
+	FActivationMetadata FindActivationMetadata() const;
 	
 	void SetCanReceiveInput(bool bInCanReceiveInput);
 	
diff --git a/Engine/Plugins/Experimental/CommonUI/Source/CommonUI/Public/CommonActivatableWidget.h b/Engine/Plugins/Experimental/CommonUI/Source/CommonUI/Public/CommonActivatableWidget.h
--- a/Engine/Plugins/Experimental/CommonUI/Source/CommonUI/Public/CommonActivatableWidget.h
+++ b/Engine/Plugins/Experimental/CommonUI/Source/CommonUI/Public/CommonActivatableWidget.h
@@ -71,10 +71,13 @@ public:
 	FSimpleMulticastDelegate& OnActivated() const { return OnActivatedEvent; }
 	FSimpleMulticastDelegate& OnDeactivated() const { return OnDeactivatedEvent; }
 
-//COMMONUI_SCOPE:
 public:
-	
-	virtual TOptional<FUICameraConfig> GetDesiredCameraConfig() const;
+
+	/**
+	 * Gets custom game-specific activation metadata for this widget. By default does nothing & used for nothing.
+	 */
+	virtual TOptional<FActivationMetadata> GetActivationMetadata() const;
+
 	/**
 	 * Gets the desired input configuration to establish when this widget activates and can receive input (i.e. all parents are also active).
 	 * This configuration will override the existing one established by any previous activatable widget and restore it (if valid) upon deactivation.
diff --git a/Engine/Plugins/Experimental/CommonUI/Source/CommonUI/Public/Input/CommonUIActionRouterBase.h b/Engine/Plugins/Experimental/CommonUI/Source/CommonUI/Public/Input/CommonUIActionRouterBase.h
--- a/Engine/Plugins/Experimental/CommonUI/Source/CommonUI/Public/Input/CommonUIActionRouterBase.h
+++ b/Engine/Plugins/Experimental/CommonUI/Source/CommonUI/Public/Input/CommonUIActionRouterBase.h
@@ -81,8 +81,8 @@ public:
 	ECommonInputMode GetActiveInputMode(ECommonInputMode DefaultInputMode = ECommonInputMode::All) const;
 	EMouseCaptureMode GetActiveMouseCaptureMode(EMouseCaptureMode DefaultMouseCapture = EMouseCaptureMode::NoCapture) const;
 
-	DECLARE_EVENT_OneParam(UCommonUIActionRouterBase, FOnCameraConfigChanged, FUICameraConfig);
-	FOnCameraConfigChanged& OnCameraConfigChanged() const { return OnCameraConfigChangedEvent; }
+	DECLARE_EVENT_OneParam(UCommonUIActionRouterBase, FOnActivationMetadataChanged, FActivationMetadata);
+	FOnActivationMetadataChanged& OnActivationMetadataChanged() const { return OnActivationMetadataChangedEvent; }
 
 	void RegisterScrollRecipient(const UWidget& ScrollableWidget);
 	void UnregisterScrollRecipient(const UWidget& ScrollableWidget);
@@ -156,7 +156,7 @@ private:
 	FActivatableTreeNodePtr FindOwningNode(const UWidget& Widget) const;
 	FActivatableTreeNodePtr FindNodeRecursive(const FActivatableTreeNodePtr& CurrentNode, const UCommonActivatableWidget& Widget) const;
 	FActivatableTreeNodePtr FindNodeRecursive(const FActivatableTreeNodePtr& CurrentNode, const TSharedPtr<SWidget>& Widget) const;
-	void SetActiveUICameraConfig(const FUICameraConfig& NewConfig);
+	void SetActiveActivationMetadata(const FActivationMetadata& NewConfig);
 	
 	void HandleActivatableWidgetRebuilding(UCommonActivatableWidget& RebuildingWidget);
 	void ProcessRebuiltWidgets();
@@ -206,7 +206,7 @@ private:
 
 	mutable FSimpleMulticastDelegate OnBoundActionsUpdatedEvent;
 	mutable FOnActiveInputModeChanged OnActiveInputModeChangedEvent;
-	mutable FOnCameraConfigChanged OnCameraConfigChangedEvent;
+	mutable FOnActivationMetadataChanged OnActivationMetadataChangedEvent;
 
 	friend class FActionRouterBindingCollection;
 	friend class FActivatableTreeNode;
diff --git a/Engine/Plugins/Experimental/CommonUI/Source/CommonUI/Public/Input/UIActionBindingHandle.h b/Engine/Plugins/Experimental/CommonUI/Source/CommonUI/Public/Input/UIActionBindingHandle.h
--- a/Engine/Plugins/Experimental/CommonUI/Source/CommonUI/Public/Input/UIActionBindingHandle.h
+++ b/Engine/Plugins/Experimental/CommonUI/Source/CommonUI/Public/Input/UIActionBindingHandle.h
@@ -53,17 +53,29 @@ private:
 	int32 RegistrationId = INDEX_NONE;
 };
 
-struct FUICameraConfig
+
+/**
+ * Metadata that can be set on activatable widgets via GetActivationMetadata & listened to via OnActivationMetadataChanged.
+ * Useful for game-specific behaviors to be triggered by activation of certain widgets generically. For example changing
+ * camera config of your game when certain widgets activate.
+ * 
+ * By default only has an enum for this metadata.
+ */
+struct FActivationMetadata
 {
-	COMMONUI_API FUICameraConfig(){}
-	FUICameraConfig(uint8 InCamera) : DesiredCamera(InCamera) {}
+	COMMONUI_API FActivationMetadata() { }
+	FActivationMetadata(uint8 InMetadataEnum) : MetadataEnum(InMetadataEnum) {}
 
-	TOptional<uint8> GetDesiredCamera() const { return DesiredCamera; }
+	TOptional<uint8> GetMetadataEnum() const { return MetadataEnum; }
 
 private:
-	TOptional<uint8> DesiredCamera;
+	TOptional<uint8> MetadataEnum;
 };
 
+/**
+ * Input Config that can be applied on widget activation. Allows for input setup  (Mouse capture, 
+ * UI-only input, move / look ignore, etc), to be controlled by widget activation.
+ */
 struct FUIInputConfig
 {
 	ECommonInputMode GetInputMode() const { return InputMode; }
