Add virtual method so that CommonInputProcessor can be overriden.

diff --git a/Engine/Plugins/Experimental/CommonUI/Source/CommonInput/Private/CommonInputPreprocessor.cpp b/Engine/Plugins/Experimental/CommonUI/Source/CommonInput/Private/CommonInputPreprocessor.cpp
new file mode 100644
--- /dev/null
+++ b/Engine/Plugins/Experimental/CommonUI/Source/CommonInput/Private/CommonInputPreprocessor.cpp
@@ -0,0 +1,264 @@
+// Copyright Epic Games, Inc. All Rights Reserved.
+
+#include "CommonInputPreprocessor.h"
+#include "CommonInputSubsystem.h"
+#include "CommonInputPrivate.h"
+#include "CommonInputTypeEnum.h"
+
+#include "Engine/Engine.h"
+#include "Engine/GameViewportClient.h"
+#include "Engine/LocalPlayer.h"
+#include "Framework/Application/SlateApplication.h"
+#include "Framework/Application/SlateUser.h"
+#include "ICommonInputModule.h"
+
+#include "Engine/Engine.h"
+#include "Engine/GameViewportClient.h"
+#include "Engine/LocalPlayer.h"
+#include "Engine/PlatformSettingsManager.h"
+#include "Engine/World.h"
+#include "EnhancedInputSubsystems.h"
+#include "Framework/Application/IInputProcessor.h"
+#include "Framework/Application/SlateApplication.h"
+#include "Framework/Application/SlateUser.h"
+#include "Misc/ConfigCacheIni.h"
+#include "HAL/PlatformStackWalk.h"
+#include "Widgets/SViewport.h"
+#include "CommonInputSettings.h"
+#include "ICommonInputModule.h"
+
+#if WITH_EDITOR
+#include "Settings/LevelEditorPlaySettings.h"
+#endif
+
+FCommonInputPreprocessor::FCommonInputPreprocessor(UCommonInputSubsystem& InCommonInputSubsystem)
+	: InputSubsystem(InCommonInputSubsystem)
+	, bIgnoreNextMove(false)
+{
+	for (uint8 InputTypeIndex = 0; InputTypeIndex < (uint8)ECommonInputType::Count; InputTypeIndex++)
+	{
+		InputMethodPermissions[InputTypeIndex] = false;
+	}
+}
+
+void FCommonInputPreprocessor::Tick(const float DeltaTime, FSlateApplication& SlateApp, TSharedRef<ICursor> Cursor)
+{
+}
+
+bool FCommonInputPreprocessor::HandleKeyDownEvent(FSlateApplication& SlateApp, const FKeyEvent& InKeyEvent)
+{
+	const ECommonInputType InputType = GetInputType(InKeyEvent.GetKey());
+	if (IsRelevantInput(SlateApp, InKeyEvent, InputType))
+	{
+		if (IsInputMethodBlocked(InputType))
+		{
+			return true;
+		}
+
+		RefreshCurrentInputMethod(InputType);
+	}
+	return false;
+}
+
+bool FCommonInputPreprocessor::HandleAnalogInputEvent(FSlateApplication& SlateApp, const FAnalogInputEvent& InAnalogInputEvent)
+{
+	const ECommonInputType InputType = GetInputType(InAnalogInputEvent.GetKey());
+	return IsRelevantInput(SlateApp, InAnalogInputEvent, InputType) && IsInputMethodBlocked(InputType);
+}
+
+bool FCommonInputPreprocessor::HandleMouseMoveEvent(FSlateApplication& SlateApp, const FPointerEvent& InPointerEvent)
+{
+	const ECommonInputType InputType = GetInputType(InPointerEvent);
+	if (IsRelevantInput(SlateApp, InPointerEvent, InputType))
+	{
+		if (bIgnoreNextMove)
+		{
+			bIgnoreNextMove = false;
+		}
+		else if (!InPointerEvent.GetCursorDelta().IsNearlyZero())
+		{
+			if (IsInputMethodBlocked(InputType))
+			{
+				return true;
+			}
+
+			RefreshCurrentInputMethod(InputType);
+		}
+	}
+	
+	return false;
+}
+
+bool FCommonInputPreprocessor::HandleMouseButtonDownEvent(FSlateApplication& SlateApp, const FPointerEvent& InPointerEvent)
+{
+	const ECommonInputType InputType = GetInputType(InPointerEvent);
+	if (IsRelevantInput(SlateApp, InPointerEvent, InputType))
+	{
+		if (IsInputMethodBlocked(InputType))
+		{
+			return true;
+		}
+		RefreshCurrentInputMethod(InputType);
+	}
+	return false;
+}
+
+bool FCommonInputPreprocessor::HandleMouseButtonDoubleClickEvent(FSlateApplication& SlateApp, const FPointerEvent& InPointerEvent)
+{
+	const ECommonInputType InputType = GetInputType(InPointerEvent);
+	if (IsRelevantInput(SlateApp, InPointerEvent, InputType))
+	{
+		if (IsInputMethodBlocked(InputType))
+		{
+			return true;
+		}
+		RefreshCurrentInputMethod(InputType);
+	}
+	return false;
+}
+
+void FCommonInputPreprocessor::SetInputTypeFilter(ECommonInputType InputType, FName InReason, bool InFilter)
+{
+	TMap<FName, bool> Reasons = FilterInputTypeWithReasons[(uint8)InputType];
+	Reasons.Add(InReason, InFilter);
+
+	bool ComputedFilter = false;
+	for (const auto& Entry : Reasons)
+	{
+		ComputedFilter |= Entry.Value;
+	}
+
+	InputMethodPermissions[(uint8)InputType] = ComputedFilter;
+}
+
+bool FCommonInputPreprocessor::IsInputMethodBlocked(ECommonInputType InputType) const
+{
+	return InputMethodPermissions[(uint8)InputType];
+}
+
+bool FCommonInputPreprocessor::IsRelevantInput(FSlateApplication& SlateApp, const FInputEvent& InputEvent, const ECommonInputType DesiredInputType)
+{
+#if WITH_EDITOR
+	// If we're stopped at a breakpoint we need for this input preprocessor to just ignore all incoming input
+	// because we're now doing stuff outside the game loop in the editor and it needs to not block all that.
+	// This can happen if you suspend input while spawning a dialog and then hit another breakpoint and then
+	// try and use the editor, you can suddenly be unable to do anything.
+	if (GIntraFrameDebuggingGameThread)
+	{
+		return false;
+	}
+#endif
+	
+	if (SlateApp.IsActive() 
+		|| SlateApp.GetHandleDeviceInputWhenApplicationNotActive() 
+		|| (ICommonInputModule::GetSettings().GetAllowOutOfFocusDeviceInput() && DesiredInputType == ECommonInputType::Gamepad))
+	{
+		const ULocalPlayer& LocalPlayer = *InputSubsystem.GetLocalPlayerChecked();
+		int32 ControllerId = LocalPlayer.GetControllerId();
+
+#if WITH_EDITOR
+		// PIE is a very special flower, especially when running two clients - we have two LocalPlayers with ControllerId 0
+		// The editor has existing shenanigans for handling this scenario, so we're using those for now
+		// Ultimately this would ideally be something the editor resolves at the SlateApplication level with a custom ISlateInputMapping or something
+		GEngine->RemapGamepadControllerIdForPIE(LocalPlayer.ViewportClient, ControllerId);
+#endif
+		return ControllerId == InputEvent.GetUserIndex();
+	}
+	
+	return false;
+}
+
+void FCommonInputPreprocessor::RefreshCurrentInputMethod(ECommonInputType InputMethod)
+{
+#if WITH_EDITOR
+	// Lots of special-case fun for PIE - there are special scenarios wherein we want to ignore the update attempt
+	const ULocalPlayer& LocalPlayer = *InputSubsystem.GetLocalPlayerChecked();
+	bool bCanApplyInputMethodUpdate = false;
+	if (LocalPlayer.ViewportClient)
+	{
+		TSharedPtr<FSlateUser> SlateUser = FSlateApplication::Get().GetUser(LocalPlayer.GetControllerId());
+		if (ensure(SlateUser))
+		{
+			bool bRoutingGamepadToNextPlayer = false;
+			if (!GetDefault<ULevelEditorPlaySettings>()->GetRouteGamepadToSecondWindow(bRoutingGamepadToNextPlayer))
+			{
+				// This looks strange - it's because the getter will set the output param based on the value of the setting, but return
+				//	whether the setting actually matters. So we're making sure that even if the setting is true, we revert to false when it's irrelevant.
+				bRoutingGamepadToNextPlayer = false;
+			}
+
+			if (SlateUser->IsWidgetInFocusPath(LocalPlayer.ViewportClient->GetGameViewportWidget()))
+			{
+				// Our owner's game viewport is in the focus path, which in a PIE scenario means we shouldn't
+				// acknowledge gamepad input if it's being routed to another PIE client
+				if (InputMethod != ECommonInputType::Gamepad || !bRoutingGamepadToNextPlayer)
+				{
+					bCanApplyInputMethodUpdate = true;
+				}
+			}
+			else if (InputMethod == ECommonInputType::Gamepad)
+			{
+				bCanApplyInputMethodUpdate = bRoutingGamepadToNextPlayer;
+			}
+		}
+	}
+	if (!bCanApplyInputMethodUpdate)
+	{
+		return;
+	}
+#endif
+
+	InputSubsystem.SetCurrentInputType(InputMethod);
+
+	// Try to auto-detect the type of gamepad
+	if ((InputMethod == ECommonInputType::Gamepad) && UCommonInputPlatformSettings::Get()->CanChangeGamepadType())
+	{
+		if (const FInputDeviceScope* DeviceScope = FInputDeviceScope::GetCurrent())
+		{
+			if ((DeviceScope->InputDeviceName != LastSeenGamepadInputDeviceName) || (DeviceScope->HardwareDeviceIdentifier != LastSeenGamepadHardwareDeviceIdentifier))
+			{
+				LastSeenGamepadInputDeviceName = DeviceScope->InputDeviceName;
+				LastSeenGamepadHardwareDeviceIdentifier = DeviceScope->HardwareDeviceIdentifier;
+
+				const FName GamepadInputType = InputSubsystem.GetCurrentGamepadName();
+				const FName BestGamepadType = UCommonInputPlatformSettings::Get()->GetBestGamepadNameForHardware(GamepadInputType, DeviceScope->InputDeviceName, DeviceScope->HardwareDeviceIdentifier);
+				if (BestGamepadType != GamepadInputType)
+				{
+					UE_LOG(LogCommonInput, Log, TEXT("UCommonInputSubsystem: Autodetect changed GamepadInputType to %s"), *BestGamepadType.ToString());
+					InputSubsystem.SetGamepadInputType(BestGamepadType);
+					OnGamepadChangeDetected.Broadcast(BestGamepadType);
+				}
+			}
+		}
+	}
+}
+
+ECommonInputType FCommonInputPreprocessor::GetInputType(const FKey& Key)
+{
+	// If the key is a gamepad key or if the key is a Click key (which simulates a click), we should be in Gamepad mode
+	if (Key.IsGamepadKey() || InputSubsystem.GetIsGamepadSimulatedClick())
+	{
+		if (UCommonInputSubsystem::IsMobileGamepadKey(Key))
+		{
+			return ECommonInputType::Touch;
+		}
+
+		//@todo DanH: Why would the gamepad simulating a click imply we also want to ignore key input? Should the property instead be IsGamepadSimulatedInput? Or is this wrong?
+		return ECommonInputType::Gamepad;
+	}
+
+	return ECommonInputType::MouseAndKeyboard;
+}
+
+ECommonInputType FCommonInputPreprocessor::GetInputType(const FPointerEvent& PointerEvent)
+{
+	if (PointerEvent.IsTouchEvent())
+	{
+		return ECommonInputType::Touch;
+	}
+	else if (InputSubsystem.GetIsGamepadSimulatedClick())
+	{
+		return ECommonInputType::Gamepad;
+	}
+	return ECommonInputType::MouseAndKeyboard;
+}
\ No newline at end of file
diff --git a/Engine/Plugins/Experimental/CommonUI/Source/CommonInput/Private/CommonInputSubsystem.cpp b/Engine/Plugins/Experimental/CommonUI/Source/CommonInput/Private/CommonInputSubsystem.cpp
--- a/Engine/Plugins/Experimental/CommonUI/Source/CommonInput/Private/CommonInputSubsystem.cpp
+++ b/Engine/Plugins/Experimental/CommonUI/Source/CommonInput/Private/CommonInputSubsystem.cpp
@@ -1,6 +1,7 @@
 // Copyright Epic Games, Inc. All Rights Reserved.
 
 #include "CommonInputSubsystem.h"
+#include "CommonInputPreprocessor.h"
 #include "CommonInputPrivate.h"
 #include "CommonInputTypeEnum.h"
 
@@ -21,10 +22,6 @@
 
 #include UE_INLINE_GENERATED_CPP_BY_NAME(CommonInputSubsystem)
 
-#if WITH_EDITOR
-#include "Settings/LevelEditorPlaySettings.h"
-#endif
-
 #if !UE_BUILD_SHIPPING
 static int32 bDumpInputTypeChangeCallstack = 0;
 static FAutoConsoleVariableRef CVarDumpInputTypeChangeCallstack(
@@ -36,266 +33,6 @@ static FAutoConsoleVariableRef CVarDumpInputTypeChangeCallstack(
 
 FPlatformInputSupportOverrideDelegate UCommonInputSubsystem::OnPlatformInputSupportOverride;
 
-/**
- * Helper class that is designed to fire before any UI has a chance to process input so that
- * we can properly set the current input type of the application.
- */
-class FCommonInputPreprocessor : public IInputProcessor
-{
-public:
-	FCommonInputPreprocessor(UCommonInputSubsystem& InCommonInputSubsystem)
-		: InputSubsystem(InCommonInputSubsystem)
-		, bIgnoreNextMove(false)
-	{
-		for (uint8 InputTypeIndex = 0; InputTypeIndex < (uint8)ECommonInputType::Count; InputTypeIndex++)
-		{
-			InputMethodPermissions[InputTypeIndex] = false;
-		}
-	}
-
-	
-	virtual void Tick(const float DeltaTime, FSlateApplication& SlateApp, TSharedRef<ICursor> Cursor) override
-	{
-	}
-
-	virtual bool HandleKeyDownEvent(FSlateApplication& SlateApp, const FKeyEvent& InKeyEvent) override
-	{
-		const ECommonInputType InputType = GetInputType(InKeyEvent.GetKey());
-		if (IsRelevantInput(SlateApp, InKeyEvent, InputType))
-		{
-			if (IsInputMethodBlocked(InputType))
-			{
-				return true;
-			}
-
-			RefreshCurrentInputMethod(InputType);
-		}
-		return false;
-	}
-
-	virtual bool HandleAnalogInputEvent(FSlateApplication& SlateApp, const FAnalogInputEvent& InAnalogInputEvent) override
-	{
-		const ECommonInputType InputType = GetInputType(InAnalogInputEvent.GetKey());
-		return IsRelevantInput(SlateApp, InAnalogInputEvent, InputType) && IsInputMethodBlocked(InputType);
-	}
-
-	virtual bool HandleMouseMoveEvent(FSlateApplication& SlateApp, const FPointerEvent& InPointerEvent) override
-	{
-		const ECommonInputType InputType = GetInputType(InPointerEvent);
-		if (IsRelevantInput(SlateApp, InPointerEvent, InputType))
-		{
-			if (bIgnoreNextMove)
-			{
-				bIgnoreNextMove = false;
-			}
-			else if (!InPointerEvent.GetCursorDelta().IsNearlyZero())
-			{
-				if (IsInputMethodBlocked(InputType))
-				{
-					return true;
-				}
-
-				RefreshCurrentInputMethod(InputType);
-			}
-		}
-		
-		return false;
-	}
-
-	virtual bool HandleMouseButtonDownEvent(FSlateApplication& SlateApp, const FPointerEvent& InPointerEvent) override
-	{
-		const ECommonInputType InputType = GetInputType(InPointerEvent);
-		if (IsRelevantInput(SlateApp, InPointerEvent, InputType))
-		{
-			if (IsInputMethodBlocked(InputType))
-			{
-				return true;
-			}
-			RefreshCurrentInputMethod(InputType);
-		}
-		return false;
-	}
-
-	virtual bool HandleMouseButtonDoubleClickEvent(FSlateApplication& SlateApp, const FPointerEvent& InPointerEvent) override
-	{
-		const ECommonInputType InputType = GetInputType(InPointerEvent);
-		if (IsRelevantInput(SlateApp, InPointerEvent, InputType))
-		{
-			if (IsInputMethodBlocked(InputType))
-			{
-				return true;
-			}
-			RefreshCurrentInputMethod(InputType);
-		}
-		return false;
-	}
-
-	virtual const TCHAR* GetDebugName() const override { return TEXT("CommonInput"); }
-
-	void SetInputTypeFilter(ECommonInputType InputType, FName InReason, bool InFilter)
-	{
-		TMap<FName, bool> Reasons = FilterInputTypeWithReasons[(uint8)InputType];
-		Reasons.Add(InReason, InFilter);
-
-		bool ComputedFilter = false;
-		for (const auto& Entry : Reasons)
-		{
-			ComputedFilter |= Entry.Value;
-		}
-
-		InputMethodPermissions[(uint8)InputType] = ComputedFilter;
-	}
-
-	bool IsInputMethodBlocked(ECommonInputType InputType) const
-	{
-		return InputMethodPermissions[(uint8)InputType];
-	}
-
-	FGamepadChangeDetectedEvent OnGamepadChangeDetected;
-
-private:
-	bool IsRelevantInput(FSlateApplication& SlateApp, const FInputEvent& InputEvent, const ECommonInputType DesiredInputType)
-	{
-#if WITH_EDITOR
-		// If we're stopped at a breakpoint we need for this input preprocessor to just ignore all incoming input
-		// because we're now doing stuff outside the game loop in the editor and it needs to not block all that.
-		// This can happen if you suspend input while spawning a dialog and then hit another breakpoint and then
-		// try and use the editor, you can suddenly be unable to do anything.
-		if (GIntraFrameDebuggingGameThread)
-		{
-			return false;
-		}
-#endif
-		
-		if (SlateApp.IsActive() 
-			|| SlateApp.GetHandleDeviceInputWhenApplicationNotActive() 
-			|| (ICommonInputModule::GetSettings().GetAllowOutOfFocusDeviceInput() && DesiredInputType == ECommonInputType::Gamepad))
-		{
-			const ULocalPlayer& LocalPlayer = *InputSubsystem.GetLocalPlayerChecked();
-			int32 ControllerId = LocalPlayer.GetControllerId();
-
-#if WITH_EDITOR
-			// PIE is a very special flower, especially when running two clients - we have two LocalPlayers with ControllerId 0
-			// The editor has existing shenanigans for handling this scenario, so we're using those for now
-			// Ultimately this would ideally be something the editor resolves at the SlateApplication level with a custom ISlateInputMapping or something
-			GEngine->RemapGamepadControllerIdForPIE(LocalPlayer.ViewportClient, ControllerId);
-#endif
-			return ControllerId == InputEvent.GetUserIndex();
-		}
-		
-		return false;
-	}
-
-	void RefreshCurrentInputMethod(ECommonInputType InputMethod)
-	{
-#if WITH_EDITOR
-		// Lots of special-case fun for PIE - there are special scenarios wherein we want to ignore the update attempt
-		const ULocalPlayer& LocalPlayer = *InputSubsystem.GetLocalPlayerChecked();
-		bool bCanApplyInputMethodUpdate = false;
-		if (LocalPlayer.ViewportClient)
-		{
-			TSharedPtr<FSlateUser> SlateUser = FSlateApplication::Get().GetUser(LocalPlayer.GetControllerId());
-			if (ensure(SlateUser))
-			{
-				bool bRoutingGamepadToNextPlayer = false;
-				if (!GetDefault<ULevelEditorPlaySettings>()->GetRouteGamepadToSecondWindow(bRoutingGamepadToNextPlayer))
-				{
-					// This looks strange - it's because the getter will set the output param based on the value of the setting, but return
-					//	whether the setting actually matters. So we're making sure that even if the setting is true, we revert to false when it's irrelevant.
-					bRoutingGamepadToNextPlayer = false;
-				}
-
-				if (SlateUser->IsWidgetInFocusPath(LocalPlayer.ViewportClient->GetGameViewportWidget()))
-				{
-					// Our owner's game viewport is in the focus path, which in a PIE scenario means we shouldn't
-					// acknowledge gamepad input if it's being routed to another PIE client
-					if (InputMethod != ECommonInputType::Gamepad || !bRoutingGamepadToNextPlayer)
-					{
-						bCanApplyInputMethodUpdate = true;
-					}
-				}
-				else if (InputMethod == ECommonInputType::Gamepad)
-				{
-					bCanApplyInputMethodUpdate = bRoutingGamepadToNextPlayer;
-				}
-			}
-		}
-		if (!bCanApplyInputMethodUpdate)
-		{
-			return;
-		}
-#endif
-
-		InputSubsystem.SetCurrentInputType(InputMethod);
-
-		// Try to auto-detect the type of gamepad
-		if ((InputMethod == ECommonInputType::Gamepad) && UCommonInputPlatformSettings::Get()->CanChangeGamepadType())
-		{
-			if (const FInputDeviceScope* DeviceScope = FInputDeviceScope::GetCurrent())
-			{
-				if ((DeviceScope->InputDeviceName != LastSeenGamepadInputDeviceName) || (DeviceScope->HardwareDeviceIdentifier != LastSeenGamepadHardwareDeviceIdentifier))
-				{
-					LastSeenGamepadInputDeviceName = DeviceScope->InputDeviceName;
-					LastSeenGamepadHardwareDeviceIdentifier = DeviceScope->HardwareDeviceIdentifier;
-
-					const FName GamepadInputType = InputSubsystem.GetCurrentGamepadName();
-					const FName BestGamepadType = UCommonInputPlatformSettings::Get()->GetBestGamepadNameForHardware(GamepadInputType, DeviceScope->InputDeviceName, DeviceScope->HardwareDeviceIdentifier);
-					if (BestGamepadType != GamepadInputType)
-					{
-						UE_LOG(LogCommonInput, Log, TEXT("UCommonInputSubsystem: Autodetect changed GamepadInputType to %s"), *BestGamepadType.ToString());
-						InputSubsystem.SetGamepadInputType(BestGamepadType);
-						OnGamepadChangeDetected.Broadcast(BestGamepadType);
-					}
-				}
-			}
-		}
-	}
-
-	ECommonInputType GetInputType(const FKey& Key)
-	{
-		// If the key is a gamepad key or if the key is a Virtual Click key (which simulates a click), we should be in Gamepad mode
-		if (Key.IsGamepadKey() || InputSubsystem.GetIsGamepadSimulatedClick())
-		{
-			if (UCommonInputSubsystem::IsMobileGamepadKey(Key))
-			{
-				return ECommonInputType::Touch;
-			}
-
-			//@todo DanH: Why would the gamepad simulating a click imply we also want to ignore key input? Should the property instead be IsGamepadSimulatedInput? Or is this wrong?
-			return ECommonInputType::Gamepad;
-		}
-
-		return ECommonInputType::MouseAndKeyboard;
-	}
-
-	ECommonInputType GetInputType(const FPointerEvent& PointerEvent)
-	{
-		if (PointerEvent.IsTouchEvent())
-		{
-			return ECommonInputType::Touch;
-		}
-		else if (InputSubsystem.GetIsGamepadSimulatedClick())
-		{
-			return ECommonInputType::Gamepad;
-		}
-		return ECommonInputType::MouseAndKeyboard;
-	}
-	
-private:
-	UCommonInputSubsystem& InputSubsystem;
-	
-	bool bIgnoreNextMove = false;
-	bool InputMethodPermissions[(uint8)ECommonInputType::Count];
-
-	// The reasons we might be filtering input right now.
-	TMap<FName, bool> FilterInputTypeWithReasons[(uint8)ECommonInputType::Count];
-
-	FName LastSeenGamepadInputDeviceName;
-	FString LastSeenGamepadHardwareDeviceIdentifier;
-
-	friend class UCommonInputSubsystem;
-};
-
 //////////////////////////////////////////////////////////////////////////
 // UCommonInputSubsystem
 //////////////////////////////////////////////////////////////////////////
@@ -348,7 +85,7 @@ void UCommonInputSubsystem::Initialize(FSubsystemCollectionBase& Collection)
 	GamepadInputType = Settings->GetDefaultGamepadName();
 	CurrentInputType = LastInputType = Settings->GetDefaultInputType();
 
-	CommonInputPreprocessor = MakeShared<FCommonInputPreprocessor>(*this);
+	CommonInputPreprocessor = MakeInputProcessor();
 	FSlateApplication::Get().RegisterInputPreProcessor(CommonInputPreprocessor, 0);
 
 	TickHandle = FTSTicker::GetCoreTicker().AddTicker(FTickerDelegate::CreateUObject(this, &UCommonInputSubsystem::Tick), 0.1f);
@@ -458,6 +195,11 @@ bool UCommonInputSubsystem::IsInputMethodActive(ECommonInputType InputMethod) co
 	return GetCurrentInputType() == InputMethod;
 }
 
+TSharedPtr<FCommonInputPreprocessor> UCommonInputSubsystem::MakeInputProcessor()
+{
+	return MakeShared<FCommonInputPreprocessor>(*this);
+}
+
 ECommonInputType UCommonInputSubsystem::LockInput(ECommonInputType InputToLock) const
 {
 	return CurrentInputLock.Get(InputToLock);
diff --git a/Engine/Plugins/Experimental/CommonUI/Source/CommonInput/Public/CommonInputPreprocessor.h b/Engine/Plugins/Experimental/CommonUI/Source/CommonInput/Public/CommonInputPreprocessor.h
new file mode 100644
--- /dev/null
+++ b/Engine/Plugins/Experimental/CommonUI/Source/CommonInput/Public/CommonInputPreprocessor.h
@@ -0,0 +1,64 @@
+// Copyright Epic Games, Inc. All Rights Reserved.
+
+#pragma once
+
+#include "Subsystems/LocalPlayerSubsystem.h"
+#include "CommonInputBaseTypes.h"
+#include "CommonInputSubsystem.h"
+#include "Containers/Ticker.h"
+#include "Framework/Application/IInputProcessor.h"
+#include "Framework/Application/SlateApplication.h"
+
+#if UE_ENABLE_INCLUDE_ORDER_DEPRECATED_IN_5_2
+#include "Framework/Application/SlateApplication.h"
+#include "Styling/SlateTypes.h"
+#endif
+
+/**
+ * Helper class that is designed to fire before any UI has a chance to process input so that
+ * we can properly set the current input type of the application.
+ */
+class COMMONINPUT_API FCommonInputPreprocessor : public IInputProcessor
+{
+public:
+	FCommonInputPreprocessor(UCommonInputSubsystem& InCommonInputSubsystem);
+
+	//~ Begin IInputProcessor Interface
+	virtual void Tick(const float DeltaTime, FSlateApplication& SlateApp, TSharedRef<ICursor> Cursor) override;
+	virtual bool HandleKeyDownEvent(FSlateApplication& SlateApp, const FKeyEvent& InKeyEvent) override;
+	virtual bool HandleAnalogInputEvent(FSlateApplication& SlateApp, const FAnalogInputEvent& InAnalogInputEvent) override;
+	virtual bool HandleMouseMoveEvent(FSlateApplication& SlateApp, const FPointerEvent& InPointerEvent) override;
+	virtual bool HandleMouseButtonDownEvent(FSlateApplication& SlateApp, const FPointerEvent& InPointerEvent) override;
+	virtual bool HandleMouseButtonDoubleClickEvent(FSlateApplication& SlateApp, const FPointerEvent& InPointerEvent) override;
+	virtual const TCHAR* GetDebugName() const override { return TEXT("CommonInput"); }
+	//~ End IInputProcessor Interface
+
+	void SetInputTypeFilter(ECommonInputType InputType, FName InReason, bool InFilter);
+
+	bool IsInputMethodBlocked(ECommonInputType InputType) const;
+
+	FGamepadChangeDetectedEvent OnGamepadChangeDetected;
+
+protected:
+	bool IsRelevantInput(FSlateApplication& SlateApp, const FInputEvent& InputEvent, const ECommonInputType DesiredInputType);
+
+	void RefreshCurrentInputMethod(ECommonInputType InputMethod);
+
+	ECommonInputType GetInputType(const FKey& Key);
+
+	ECommonInputType GetInputType(const FPointerEvent& PointerEvent);
+	
+protected:
+	UCommonInputSubsystem& InputSubsystem;
+	
+	bool bIgnoreNextMove = false;
+	bool InputMethodPermissions[(uint8)ECommonInputType::Count];
+
+	// The reasons we might be filtering input right now.
+	TMap<FName, bool> FilterInputTypeWithReasons[(uint8)ECommonInputType::Count];
+
+	FName LastSeenGamepadInputDeviceName;
+	FString LastSeenGamepadHardwareDeviceIdentifier;
+
+	friend class UCommonInputSubsystem;
+};
\ No newline at end of file
diff --git a/Engine/Plugins/Experimental/CommonUI/Source/CommonInput/Public/CommonInputSubsystem.h b/Engine/Plugins/Experimental/CommonUI/Source/CommonInput/Public/CommonInputSubsystem.h
--- a/Engine/Plugins/Experimental/CommonUI/Source/CommonInput/Public/CommonInputSubsystem.h
+++ b/Engine/Plugins/Experimental/CommonUI/Source/CommonInput/Public/CommonInputSubsystem.h
@@ -106,11 +106,16 @@ public:
 	static FPlatformInputSupportOverrideDelegate& GetOnPlatformInputSupportOverride() { return OnPlatformInputSupportOverride; }
 
 protected:
+	virtual TSharedPtr<FCommonInputPreprocessor> MakeInputProcessor();
+
 	ECommonInputType LockInput(ECommonInputType InputToLock) const;
 
 	UFUNCTION()
 	void BroadcastInputMethodChanged();
 
+protected:
+	TSharedPtr<FCommonInputPreprocessor> CommonInputPreprocessor;
+
 private:
 	bool Tick(float DeltaTime);
 
@@ -155,8 +160,6 @@ private:
 
 	TOptional<ECommonInputType> CurrentInputLock;
 
-	TSharedPtr<FCommonInputPreprocessor> CommonInputPreprocessor;
-
 	UPROPERTY(Transient)
 	TObjectPtr<UCommonInputActionDomainTable> ActionDomainTable;
 
