Clean up comments on CommonUI

diff --git a/Engine/Plugins/Experimental/CommonUI/Source/CommonInput/Private/CommonInputPreprocessor.cpp b/Engine/Plugins/Experimental/CommonUI/Source/CommonInput/Private/CommonInputPreprocessor.cpp
--- a/Engine/Plugins/Experimental/CommonUI/Source/CommonInput/Private/CommonInputPreprocessor.cpp
+++ b/Engine/Plugins/Experimental/CommonUI/Source/CommonInput/Private/CommonInputPreprocessor.cpp
@@ -242,8 +242,7 @@ ECommonInputType FCommonInputPreprocessor::GetInputType(const FKey& Key)
 		{
 			return ECommonInputType::Touch;
 		}
-
-		//@todo DanH: Why would the gamepad simulating a click imply we also want to ignore key input? Should the property instead be IsGamepadSimulatedInput? Or is this wrong?
+		
 		return ECommonInputType::Gamepad;
 	}
 
diff --git a/Engine/Plugins/Experimental/CommonUI/Source/CommonInput/Private/CommonInputSubsystem.cpp b/Engine/Plugins/Experimental/CommonUI/Source/CommonInput/Private/CommonInputSubsystem.cpp
--- a/Engine/Plugins/Experimental/CommonUI/Source/CommonInput/Private/CommonInputSubsystem.cpp
+++ b/Engine/Plugins/Experimental/CommonUI/Source/CommonInput/Private/CommonInputSubsystem.cpp
@@ -61,17 +61,7 @@ UCommonInputSubsystem* UCommonInputSubsystem::Get(const ULocalPlayer* LocalPlaye
 
 UCommonInputSubsystem::UCommonInputSubsystem()
 {
-	//@TODO: Project setting (should be done as another config var any platform can set tho)?
-//	// Uncomment this if we want mobile platforms to default to gamepad when gamepads are enabled.
-// #if PLATFORM_IOS
-//     bool bAllowControllers = false;
-//     GConfig->GetBool(TEXT("/Script/IOSRuntimeSettings.IOSRuntimeSettings"), TEXT("bAllowControllers"), bAllowControllers, GEngineIni);
-//     CurrentInputType = LastInputType = bAllowControllers ? ECommonInputType::Gamepad : CurrentInputType;
-// #elif PLATFORM_ANDROID
-//     bool bAllowControllers = false;
-//     GConfig->GetBool(TEXT("/Script/AndroidRuntimeSettings.AndroidRuntimeSettings"), TEXT("bAllowControllers"), bAllowControllers, GEngineIni);
-//     CurrentInputType = LastInputType = bAllowControllers ? ECommonInputType::Gamepad : CurrentInputType;
-// #endif
+
 }
 
 void UCommonInputSubsystem::Initialize(FSubsystemCollectionBase& Collection)
@@ -391,12 +381,9 @@ bool UCommonInputSubsystem::Tick(float DeltaTime)
 {
 	QUICK_SCOPE_CYCLE_COUNTER(STAT_UCommonInputSubsystem_Tick);
 
-	//@todo DanH: This is wrong now that two of these might exist (and has always been wrong for multi-client PIE scenarios)
-	//		Preprocessors need to be kept associated with their registration priority so we can safely know these won't get all thrown out of whack as others come and go
-
-	//Keep the CommonInputPreprocessor on top. Input swap and input filtering (e.g. "Ignore Gamepad Input")
-	//both start to break down if narrow game preprocessors temporarily get in front of it.
-	//This is a workaround to avoid a bigger intervention in the SlateApplication API for managing preprocessors.
+	// Keep the CommonInputPreprocessor on top. Input swap and input filtering (e.g. "Ignore Gamepad Input")
+	// both start to break down if narrow game preprocessors temporarily get in front of it.
+	// This is a workaround to avoid a bigger intervention in the SlateApplication API for managing preprocessors.
 	if (CommonInputPreprocessor.IsValid() && FSlateApplication::IsInitialized())
 	{
 		FSlateApplication& SlateApplication = FSlateApplication::Get();
diff --git a/Engine/Plugins/Experimental/CommonUI/Source/CommonUI/Private/CommonActionWidget.cpp b/Engine/Plugins/Experimental/CommonUI/Source/CommonUI/Private/CommonActionWidget.cpp
--- a/Engine/Plugins/Experimental/CommonUI/Source/CommonUI/Private/CommonActionWidget.cpp
+++ b/Engine/Plugins/Experimental/CommonUI/Source/CommonUI/Private/CommonActionWidget.cpp
@@ -213,7 +213,6 @@ void UCommonActionWidget::SetInputAction(FDataTableRowHandle InputActionRow)
 
 void UCommonActionWidget::SetInputActionBinding(FUIActionBindingHandle BindingHandle)
 {
-	//@todo DanH/(josh.gross): Just handling the legacy stuff for now
 	UpdateBindingHandleInternal(BindingHandle);
 	if (TSharedPtr<FUIActionBinding> Binding = FUIActionBinding::FindBinding(BindingHandle))
 	{
diff --git a/Engine/Plugins/Experimental/CommonUI/Source/CommonUI/Private/CommonActivatableWidget.cpp b/Engine/Plugins/Experimental/CommonUI/Source/CommonUI/Private/CommonActivatableWidget.cpp
--- a/Engine/Plugins/Experimental/CommonUI/Source/CommonUI/Private/CommonActivatableWidget.cpp
+++ b/Engine/Plugins/Experimental/CommonUI/Source/CommonUI/Private/CommonActivatableWidget.cpp
@@ -282,7 +282,6 @@ void UCommonActivatableWidget::NativeOnDeactivated()
 
 bool UCommonActivatableWidget::NativeOnHandleBackAction()
 {
-	//@todo DanH: This isn't actually fleshed out enough - we need to figure out whether we want to let back be a normal action or a special one routed through all activatables that lets them conditionally handle it
 	if (bIsBackHandler)
 	{
 		if (!BP_OnHandleBackAction())
diff --git a/Engine/Plugins/Experimental/CommonUI/Source/CommonUI/Private/CommonGameViewportClient.cpp b/Engine/Plugins/Experimental/CommonUI/Source/CommonUI/Private/CommonGameViewportClient.cpp
--- a/Engine/Plugins/Experimental/CommonUI/Source/CommonUI/Private/CommonGameViewportClient.cpp
+++ b/Engine/Plugins/Experimental/CommonUI/Source/CommonUI/Private/CommonGameViewportClient.cpp
@@ -182,7 +182,6 @@ void UCommonGameViewportClient::HandleRerouteTouch(int32 ControllerId, uint32 To
 			UCommonUIActionRouterBase* ActionRouter = LocalPlayer->GetSubsystem<UCommonUIActionRouterBase>();
 			if (ensure(ActionRouter))
 			{
-				//@todo DanH: Does anyone actually use this? Do we need to support holds or something with this?
 				EInputEvent SimilarInputEvent = IE_MAX;
 				switch (TouchType)
 				{
diff --git a/Engine/Plugins/Experimental/CommonUI/Source/CommonUI/Private/CommonLoadGuard.cpp b/Engine/Plugins/Experimental/CommonUI/Source/CommonUI/Private/CommonLoadGuard.cpp
--- a/Engine/Plugins/Experimental/CommonUI/Source/CommonUI/Private/CommonLoadGuard.cpp
+++ b/Engine/Plugins/Experimental/CommonUI/Source/CommonUI/Private/CommonLoadGuard.cpp
@@ -351,7 +351,8 @@ void UCommonLoadGuard::PostLoad()
 	}
 
 #if WITH_EDITOR
-	//TODO: I want this removed, it's needed for backwards compatibility
+
+	// Needed for backwards compatibility of old style updating
 	if (!TextStyle && !bStyleNoLongerNeedsConversion && !IsRunningDedicatedServer())
 	{
 		UCommonUIEditorSettings& Settings = ICommonUIModule::GetEditorSettings();
diff --git a/Engine/Plugins/Experimental/CommonUI/Source/CommonUI/Private/CommonRichTextBlock.cpp b/Engine/Plugins/Experimental/CommonUI/Source/CommonUI/Private/CommonRichTextBlock.cpp
--- a/Engine/Plugins/Experimental/CommonUI/Source/CommonUI/Private/CommonRichTextBlock.cpp
+++ b/Engine/Plugins/Experimental/CommonUI/Source/CommonUI/Private/CommonRichTextBlock.cpp
@@ -83,7 +83,6 @@ public:
 			const FString& Value = TagValuePair.Value;
 			if (Tag == SupportedMarkupKeys::Color)
 			{
-				//@todo DanH: A system for named colors would be nice for sure
 				TextStyle.SetColorAndOpacity(FLinearColor::FromSRGBColor(FColor::FromHex(Value)));
 			}
 			else if (Tag == SupportedMarkupKeys::FontFace)
@@ -112,10 +111,6 @@ public:
 				{
 					TextStyle.Font.FontMaterial = nullptr;
 				}
-				else
-				{
-					//@todo DanH: Set up something to allow font materials to be assigned to lookup keys
-				}
 			}
 		}
 
@@ -215,7 +210,6 @@ public:
 				SNew(SHorizontalBox)
 				+ SHorizontalBox::Slot()
 				.AutoWidth()
-				//@todo DanH: If needed, we can have the padding between the icon and the text be auto-adjusted based on the font size
 				.Padding(0.f, 0.f, 8.f, 0.f)
 				[
 					CreateIconWidget(Entry, TextStyle, InArgs._IconScale, InArgs._bTintIcon)
@@ -333,14 +327,12 @@ protected:
 		const FInlineIconEntry* InlineIconEntry = InlineIconsByTag.Find(InlineIconEntryName);
 		if (ensure(InlineIconEntry))
 		{
-			//@todo DanH: Leave this in? Definitely of general use in any situation where you have prices in both fake icon-identified currency and real font-supported currency
-
 			// Real money symbols are within fonts and will never have an icon, so we force text only mode now 
 			// to make sure we don't have an errant icon potentially appear (or a pink error box in non-shipping)
 			// Ideally, this could be more data-driven, but that would be overkill for what is currently a singleton use case
 			ERichTextInlineIconDisplayMode EntryDisplayMode = InlineIconEntryName == TEXT("cash") ? ERichTextInlineIconDisplayMode::TextOnly : Owner->InlineIconDisplayMode;
 
-			//@todo DanH: Set something up for when to tint the icon the same color as the text
+			// Note: Icon may not be tinted same color as text, for now must be addressed in game
 			return SNew(SCommonInlineIcon, *InlineIconEntry, ModifiedTextStyle)
 				.DisplayMode(EntryDisplayMode)
 				.IconScale(IconScale)
@@ -380,103 +372,6 @@ void UCommonRichTextBlock::ReleaseSlateResources(bool bReleaseChildren)
 	MyTextScroller.Reset();
 }
 
-//@todo DanH: Address the todos nested below to enable input, and custom widgets (this is stuff originally from Orion)
-/*
-class FCommonDecorator_Input : public FCommonRichTextDecorator
-{
-public:
-	static TSharedRef<FCommonDecorator_Input> CreateInstance(UCommonRichTextBlock& InOwner)
-	{
-		return MakeShareable(new FCommonDecorator_Input(InOwner));
-	}
-
-	virtual bool Supports(const FTextRunParseResults& RunParseResult, const FString& Text) const override
-	{
-		return RunParseResult.Name == UOrionGameUIData::Get().GetRichTextStyleData().InputMarkupTag;
-	}
-
-protected:
-	virtual TSharedPtr<SWidget> CreateDecoratorWidget(const FTextRunInfo& RunInfo, const FTextBlockStyle& TextStyle) const override
-	{
-		FName AxisName;
-		FName ActionName;
-		FKey SpecificKey;
-
-		if (const FString* ActionMeta = RunInfo.MetaData.Find(SupportedMarkupKeys::Action))
-		{
-			ActionName = **ActionMeta;
-		}
-		else if (const FString* AxisMeta = RunInfo.MetaData.Find(SupportedMarkupKeys::Axis))
-		{
-			AxisName = **AxisMeta;
-		}
-		else if (const FString* KeyMeta = RunInfo.MetaData.Find(SupportedMarkupKeys::Key))
-		{
-			SpecificKey = FKey(**KeyMeta);
-		}
-		
-		const TSharedRef<FSlateFontMeasure> FontMeasure = FSlateApplication::Get().GetRenderer()->GetFontMeasureService();
-		int32 MaxCharacterHeight = FontMeasure->GetMaxCharacterHeight(TextStyle.Font);
-		
-		//@todo DanH: The common input stuff is pretty difficult to inspect or track down natively - need to sort out how to get that (ideally for free)
-		//@todo DanH: Compare the SOrionInputVisualizer to the CommonActionWidget and fill any gaps (like having an SWidget version)
-		return SNew(SOrionInputVisualizer, Owner.IsValid() ? Cast<AOrionPlayerController_Base>(Owner->GetOwningPlayer()) : nullptr)
-			.InputAction(ActionName)
-			.InputAxis(AxisName)
-			.SpecificKey(SpecificKey)
-			.bUseKeyBorder(!RunInfo.MetaData.Contains(SupportedMarkupKeys::Borderless))
-			.DesiredHeight(MaxCharacterHeight);
-	}
-
-private:
-	FCommonDecorator_Input(UCommonRichTextBlock& InOwner)
-		: FCommonRichTextDecorator(InOwner)
-	{}
-};
-
-//@todo DanH: Give this a second glance to make sure I like it, then make the CommonRichTextInlineWidget class and somewhere to register them 
-//		Also, consider whether it should be a class or an interface
-class FCommonDecorator_CustomWidget : public FCommonRichTextDecorator
-{
-public:
-	static TSharedRef<FCommonDecorator_CustomWidget> CreateInstance(UCommonRichTextBlock& InOwner)
-	{
-		return MakeShareable(new FCommonDecorator_CustomWidget(InOwner));
-	}
-
-	virtual bool Supports(const FTextRunParseResults& RunParseResult, const FString& Text) const override
-	{	
-		return UOrionGameUIData::Get().GetRichTextStyleData().InlineWidgetsByTag.Contains(RunParseResult.Name.ToLower());
-	}
-
-protected:
-	virtual TSharedPtr<SWidget> CreateDecoratorWidget(const FTextRunInfo& RunInfo, const FTextBlockStyle& TextStyle) const override
-	{
-		if (const TSubclassOf<UCommonRichTextBlockInlineWidget>* WidgetClass = UOrionGameUIData::Get().GetRichTextStyleData().InlineWidgetsByTag.Find(RunInfo.Name.ToLower()))
-		{
-			UCommonRichTextBlockInlineWidget* NewWidget = UIUtils::ConstructSubWidget(*Owner, *WidgetClass);
-
-			TSharedRef<SWidget> SlateWidget = NewWidget->TakeWidget();
-			NewWidget->InitializeInlineWidget(TextStyle, RunInfo);
-
-			return SlateWidget;
-		}
-		return TSharedPtr<SWidget>();
-	}
-
-private:
-	FCommonDecorator_CustomWidget(UCommonRichTextBlock& InOwner)
-		: FCommonRichTextDecorator(InOwner)
-	{}
-};
-
-template <typename CommonDecoratorT, typename = typename TEnableIf<TIsDerivedFrom<CommonDecoratorT, FCommonRichTextDecorator>::IsDerived, CommonDecoratorT>::Type>
-TSharedRef<ITextDecorator> CreateDecorator(UCommonRichTextBlock& InOwner)
-{
-	return CommonDecoratorT::CreateInstance(InOwner);
-}
-*/
-
 void UCommonRichTextBlock::Serialize(FArchive& Ar)
 {
 	Super::Serialize(Ar);
diff --git a/Engine/Plugins/Experimental/CommonUI/Source/CommonUI/Private/CommonTabListWidgetBase.cpp b/Engine/Plugins/Experimental/CommonUI/Source/CommonUI/Private/CommonTabListWidgetBase.cpp
--- a/Engine/Plugins/Experimental/CommonUI/Source/CommonUI/Private/CommonTabListWidgetBase.cpp
+++ b/Engine/Plugins/Experimental/CommonUI/Source/CommonUI/Private/CommonTabListWidgetBase.cpp
@@ -162,7 +162,6 @@ void UCommonTabListWidgetBase::RemoveAllTabs()
 
 int32 UCommonTabListWidgetBase::GetTabCount() const
 {
-	// TODO Should we ensure all the tabs in the list are valid?
 	return RegisteredTabsByID.Num();
 }
 
diff --git a/Engine/Plugins/Experimental/CommonUI/Source/CommonUI/Private/CommonTextBlock.cpp b/Engine/Plugins/Experimental/CommonUI/Source/CommonUI/Private/CommonTextBlock.cpp
--- a/Engine/Plugins/Experimental/CommonUI/Source/CommonUI/Private/CommonTextBlock.cpp
+++ b/Engine/Plugins/Experimental/CommonUI/Source/CommonUI/Private/CommonTextBlock.cpp
@@ -287,7 +287,7 @@ void UCommonTextBlock::PostLoad()
 void UCommonTextBlock::Serialize(FArchive& Ar)
 {
 #if WITH_EDITORONLY_DATA
-	// @fixme If a text style is set, clear style related properties on PreSave when cooking to avoid 
+	// @FIXME: If a text style is set, clear style related properties on PreSave when cooking to avoid 
 	// non-deterministic cooking issues. The root cause of this is still unclear, but we're sometimes 
 	// getting either the wrong style CDO (a parent BP instead of the child BP) or the child BP CDO properties
 	// are being initialized with the parent's values.
diff --git a/Engine/Plugins/Experimental/CommonUI/Source/CommonUI/Private/CommonVisibilityWidgetBase.cpp b/Engine/Plugins/Experimental/CommonUI/Source/CommonUI/Private/CommonVisibilityWidgetBase.cpp
--- a/Engine/Plugins/Experimental/CommonUI/Source/CommonUI/Private/CommonVisibilityWidgetBase.cpp
+++ b/Engine/Plugins/Experimental/CommonUI/Source/CommonUI/Private/CommonVisibilityWidgetBase.cpp
@@ -16,7 +16,6 @@ UDEPRECATED_UCommonVisibilityWidgetBase::UDEPRECATED_UCommonVisibilityWidgetBase
 	, VisibleType(ESlateVisibility::SelfHitTestInvisible)
 	, HiddenType(ESlateVisibility::Collapsed)
 {
-	//@TODO: The duplication of FNames is a bit of a memory waste.
 	for (const FName& RegisteredPlatform : FCommonInputPlatformBaseData::GetRegisteredPlatforms())
 	{
 		VisibilityControls.Add(RegisteredPlatform, false);
diff --git a/Engine/Plugins/Experimental/CommonUI/Source/CommonUI/Private/Input/CommonAnalogCursor.cpp b/Engine/Plugins/Experimental/CommonUI/Source/CommonUI/Private/Input/CommonAnalogCursor.cpp
--- a/Engine/Plugins/Experimental/CommonUI/Source/CommonUI/Private/Input/CommonAnalogCursor.cpp
+++ b/Engine/Plugins/Experimental/CommonUI/Source/CommonUI/Private/Input/CommonAnalogCursor.cpp
@@ -19,10 +19,6 @@
 
 #define LOCTEXT_NAMESPACE "CommonAnalogCursor"
 
-
-//@todo DanH: CVar for forcing analog movement to be enabled
-
-//@todo DanH: Move to UCommonUIInputSettings
 const float AnalogScrollUpdatePeriod = 0.1f;
 const float ScrollDeadZone = 0.2f;
 
@@ -60,8 +56,7 @@ extern bool IsViewportWindowInFocusPath(const UCommonUIActionRouterBase& Router)
 
 void FCommonAnalogCursor::Tick(const float DeltaTime, FSlateApplication& SlateApp, TSharedRef<ICursor> Cursor)
 {
-	//@todo DanH: Cursor visibility was getting thrown off somehow on PS4 and P2 wound up permanently showing the cursor
-	//		Will circle back on this, but for now refreshing a single bool each frame is relatively harmless.
+	// Refreshing visibility per tick to address multiplayer p2 cursor visibility getting stuck
 	RefreshCursorVisibility();
 
 	// Don't bother trying to do anything while the game viewport has capture
@@ -90,7 +85,6 @@ void FCommonAnalogCursor::Tick(const float DeltaTime, FSlateApplication& SlateAp
 			TSharedPtr<SWidget> CursorTarget = SlateUser->GetFocusedWidget();
 			if (TSharedPtr<ITableViewMetadata> TableViewMetadata = CursorTarget ? CursorTarget->GetMetaData<ITableViewMetadata>() : nullptr)
 			{
-				//@todo DanH: When a list is focused but the selected row isn't visible, should we try to hide the cursor or anything?
 				// A list view is currently focused, so we actually want to make sure we are centering the cursor over the currently selected row instead
 				TArray<TSharedPtr<ITableRow>> SelectedRows = TableViewMetadata->GatherSelectedRows();
 				if (SelectedRows.Num() > 0 && ensure(SelectedRows[0].IsValid()))
@@ -125,8 +119,6 @@ void FCommonAnalogCursor::Tick(const float DeltaTime, FSlateApplication& SlateAp
 					FGeometry TargetGeometry; 
 					if (CursorTarget == GetViewportClient()->GetGameViewportWidget())
 					{
-						//@todo DanH: We reeeeally need the GameViewport stuff to be friendlier toward splitscreen scenarios and allow easier direct access to a player's actual "viewport" widget
-
 						// When the target is the game viewport as a whole, we don't want to center blindly - we want to center in the geometry of our owner's widget host layer
 						TSharedPtr<IGameLayerManager> GameLayerManager = GetViewportClient()->GetGameLayerManager();
 						if (ensure(GameLayerManager))
@@ -250,8 +242,7 @@ bool FCommonAnalogCursor::HandleKeyDownEvent(FSlateApplication& SlateApp, const
 		}
 		else
 		{
-			//@todo DanH: This is a major bummer to have to flip this flag on the input subsystem here, but there is no awareness on a mouse event of whether it's real or not
-			//		Though tbh, any place that cares should be able to just check the live input mode and infer from that whether this is a mouse click or virtual gamepad click
+			// There is no awareness on a mouse event of whether it's real or not, so mark that here.
 			UCommonInputSubsystem& InputSubsytem = ActionRouter.GetInputSubsystem();
 			InputSubsytem.SetIsGamepadSimulatedClick(bIsVirtualAccept);
 			bool bReturnValue = FAnalogCursor::HandleKeyDownEvent(SlateApp, InKeyEvent);
@@ -268,7 +259,6 @@ bool FCommonAnalogCursor::HandleKeyUpEvent(FSlateApplication& SlateApp, const FK
 	if (IsRelevantInput(InKeyEvent))
 	{
 #if !UE_BUILD_SHIPPING
-		//@todo DanH: I'm not sure this'll actually work, may have been a dumb idea to try tracking this way
 		const FKey& PressedKey = InKeyEvent.GetKey();
 		if (PressedKey == EKeys::Gamepad_LeftShoulder) { ShoulderButtonStatus ^= EShoulderButtonFlags::LeftShoulder; }
 		if (PressedKey == EKeys::Gamepad_RightShoulder) { ShoulderButtonStatus ^= EShoulderButtonFlags::RightShoulder; }
@@ -341,7 +331,6 @@ bool FCommonAnalogCursor::HandleMouseButtonDownEvent(FSlateApplication& SlateApp
 		}
 #endif 
 
-		//@todo DanH: May want to make the list of "mouse buttons to treat like keys" a settings thing
 		// Mouse buttons other than the two primaries are fair game for binding as if they were normal keys
 		const FKey EffectingButton = PointerEvent.GetEffectingButton();
 		if (EffectingButton.IsMouseButton()
@@ -352,17 +341,9 @@ bool FCommonAnalogCursor::HandleMouseButtonDownEvent(FSlateApplication& SlateApp
 			UGameViewportClient* ViewportClient = GetViewportClient();
 			if (TSharedPtr<SWidget> ViewportWidget = ViewportClient ? ViewportClient->GetGameViewportWidget() : nullptr)
 			{
-				//@todo DanH: Mouse buttons generally transfer focus to the application they're over, so shouldn't we be able to rely instead
-				//		on a combination check that a) the SlateApp is active and b) the game viewport is in the focus path?
-				//		Mousing over to a web browser and clicking the back button, for example, will shift OS focus to the browser
-				//			Within the Slate app, clicking back will transfer focus just like LMB would
 				const FWidgetPath WidgetsUnderCursor = SlateApp.LocateWindowUnderMouse(PointerEvent.GetScreenSpacePosition(), SlateApp.GetInteractiveTopLevelWindows());
 				if (WidgetsUnderCursor.ContainsWidget(ViewportWidget.Get()))
 				{
-					//@todo DanH: Do we need to go through the whole process here of generating a false key down event?
-					//		All we really want to do here is allow UI input bindings to mouse buttons right?
-					//		95% sure we can/should simply be having the ActionRouter process input here instead of generating the event
-					//			After all, a widget expecting to receive an OnKeyDown with a mouse button key is just plaing doing it wrong...
 					FKeyEvent MouseKeyEvent(EffectingButton, PointerEvent.GetModifierKeys(), PointerEvent.GetUserIndex(), false, 0, 0);
 					if (SlateApp.ProcessKeyDownEvent(MouseKeyEvent))
 					{
@@ -389,10 +370,6 @@ bool FCommonAnalogCursor::HandleMouseButtonUpEvent(FSlateApplication& SlateApp,
 			bool bHandled = SlateApp.ProcessKeyUpEvent(MouseKeyEvent);
 			if (bHadKeyDown)
 			{
-				//@todo DanH: What is this scenario? As it is we'll ignore the button down when its outside the app window, but still handle it when released. NONSENSE.
-				//		Also if someone had it down when they activated the app, then released, there's also no reason to process that
-				//		Now, if we were to have a mouse button hold action (ewwwww) and you have it down and leave the app, then release, I guess we'd still process it so long as it was down
-
 				// Only block the mouse up if the mouse down was also blocked
 				return bHandled;
 			}
diff --git a/Engine/Plugins/Experimental/CommonUI/Source/CommonUI/Private/Input/CommonBoundActionBar.cpp b/Engine/Plugins/Experimental/CommonUI/Source/CommonUI/Private/Input/CommonBoundActionBar.cpp
--- a/Engine/Plugins/Experimental/CommonUI/Source/CommonUI/Private/Input/CommonBoundActionBar.cpp
+++ b/Engine/Plugins/Experimental/CommonUI/Source/CommonUI/Private/Input/CommonBoundActionBar.cpp
@@ -84,10 +84,6 @@ void UCommonBoundActionBar::OnWidgetRebuilt()
 void UCommonBoundActionBar::SynchronizeProperties()
 {
 	Super::SynchronizeProperties();
-
-	//@todo DanH: Preview some entries - this should be really easy to get out of the box without copying DynamicEntryBox
-	//	-> should make a GeneratePreviewEntries util on the base that takes a class and the number to show
-	//GeneratePreviewEntries(NumPreviewEntries, ActionButtonClass);
 }
 
 void UCommonBoundActionBar::ReleaseSlateResources(bool bReleaseChildren)
@@ -188,7 +184,7 @@ void UCommonBoundActionBar::HandleDeferredDisplayUpdate()
 								}
 								else
 								{
-									return false; //@todo(josh.gross) - allow non-legacy bindings
+									return false; 
 								}
 
 								bool bAlreadyAccepted = false;
diff --git a/Engine/Plugins/Experimental/CommonUI/Source/CommonUI/Private/Input/CommonBoundActionButton.cpp b/Engine/Plugins/Experimental/CommonUI/Source/CommonUI/Private/Input/CommonBoundActionButton.cpp
--- a/Engine/Plugins/Experimental/CommonUI/Source/CommonUI/Private/Input/CommonBoundActionButton.cpp
+++ b/Engine/Plugins/Experimental/CommonUI/Source/CommonUI/Private/Input/CommonBoundActionButton.cpp
@@ -21,7 +21,6 @@ void UCommonBoundActionButton::SetRepresentedAction(FUIActionBindingHandle InBin
 	BindingHandle = InBindingHandle;
 	UpdateInputActionWidget();
 
-	//@todo DanH: Need to figure out a better way to do this so we don't have to go mining into the (supposedly private/hidden) FUIActionBinding itself
 	if (TSharedPtr<FUIActionBinding> NewBinding = FUIActionBinding::FindBinding(InBindingHandle))
 	{
 		NewBinding->OnHoldActionProgressed.AddUObject(this, &UCommonBoundActionButton::NativeOnActionProgress);
diff --git a/Engine/Plugins/Experimental/CommonUI/Source/CommonUI/Private/Input/CommonUIActionRouterBase.cpp b/Engine/Plugins/Experimental/CommonUI/Source/CommonUI/Private/Input/CommonUIActionRouterBase.cpp
--- a/Engine/Plugins/Experimental/CommonUI/Source/CommonUI/Private/Input/CommonUIActionRouterBase.cpp
+++ b/Engine/Plugins/Experimental/CommonUI/Source/CommonUI/Private/Input/CommonUIActionRouterBase.cpp
@@ -32,7 +32,6 @@ static const FAutoConsoleVariableRef CVarAlwaysShowCursor(
 	bAlwaysShowCursor,
 	TEXT(""));
 
-//@todo DanH: TEMP LOCATION
 FGlobalUITags FGlobalUITags::GUITags;
 
 UCommonActivatableWidget* FindOwningActivatable(const UWidget& Widget)
@@ -84,7 +83,6 @@ UCommonActivatableWidget* UCommonUIActionRouterBase::FindOwningActivatable(TShar
 
 	while (Widget && !OwningActivatable)
 	{
-		//@todo DanH: Create FActivatableWidgetMetaData and slap it onto the RebuildWidget result in CommonActivatableWidget
 		Widget = Widget->GetParentWidget();
 		if (Widget && Widget->GetType().IsEqual(TEXT("SObjectWidget")))
 		{
@@ -316,7 +314,7 @@ ERouteUIInputResult UCommonUIActionRouterBase::ProcessInput(FKey Key, EInputEven
 	// In PIE, check if the user is attempting to press the StopPlaySession command chord.
 	if (GIsPlayInEditorWorld && InputEvent == IE_Pressed)
 	{
-		//TODO This could be more generic to be a list of command in commands to allow to be ignored by the UI action router.
+		// @TODO: This could be more generic to be a list of command in commands to allow to be ignored by the UI action router.
 		TSharedPtr<FUICommandInfo> StopCommand = FInputBindingManager::Get().FindCommandInContext("PlayWorld", "StopPlaySession");
 		if (ensure(StopCommand))
 		{
@@ -576,10 +574,8 @@ void UCommonUIActionRouterBase::HandleRootWidgetSlateReleased(TWeakPtr<FActivata
 		// Cannot actually have this ensure here, because we may be in a function called on the ToBeRemoved node itself, keeping one remaining strong reference
 		// ensureAlways(!ToBeRemoved.IsValid());
 
-		//@todo DanH: This may not ever actually happen, since we'll likely want the loading screen to be an activatable - we should be listening for map changes instead
 		if (RootNodes.Num() == 0)
 		{
-			//@todo DanH: This won't actually change the current config, which we may want to do with a loading screen
 			ActiveInputConfig.Reset();
 			RefreshActionDomainLeafNodeConfig();
 		}
@@ -894,7 +890,6 @@ void UCommonUIActionRouterBase::ProcessRebuiltWidgets()
 
 	if (WidgetsByDirectParent.Num() > 0)
 	{
-		//@todo DanH: Build a string to print all the remaining entries
 		ensureAlwaysMsgf(false, TEXT("Somehow we rebuilt a widget that is owned by an activatable, but no node exists for that activatable. This *should* be completely impossible."));
 	}
 	
@@ -1335,7 +1330,7 @@ void UCommonUIActionRouterBase::ApplyUIInputConfig(const FUIInputConfig& NewConf
 
 		ULocalPlayer& LocalPlayer = *GetLocalPlayerChecked();
 
-		//@todo DanH: This won't quite work for splitscreen - we need per-player viewport client settings for mouse capture
+		// Note: may not work for splitscreen. We need per-player viewport client settings for mouse capture
 		if (UGameViewportClient* GameViewportClient = LocalPlayer.ViewportClient)
 		{
 			if (TSharedPtr<SViewport> ViewportWidget = GameViewportClient->GetGameViewportWidget())
@@ -1399,9 +1394,6 @@ void UCommonUIActionRouterBase::ApplyUIInputConfig(const FUIInputConfig& NewConf
 									break;
 								// Gamepad - Let the settings tell us if we should center it.
 								case ECommonInputType::Gamepad:
-									//TODO - Consider doing something here so that the gamepad isn't moved when the gamepad's
-									// hovering is really based on what has focus right now when not in analog mode, but we
-									// don't have a good way to know that here, other than to maybe expose a setting.
 									break;
 							}
 
@@ -1569,7 +1561,6 @@ public:
 	}
 };
 
-//@todo DanH: Debug output for this stuff - Cheatscript? ShowDebug? Full monitor window a-la WidgetReflector?
 static const FAutoConsoleCommandWithWorldAndArgs DumpActivatableTreeCommand(
 	TEXT("CommonUI.DumpActivatableTree"),
 	TEXT("Outputs the current state of the activatable tree. 4 args: bIncludeActions, bIncludeChildren, bIncludeInactive, LocalPlayerId (optional, defaults to -1 or all)"),
diff --git a/Engine/Plugins/Experimental/CommonUI/Source/CommonUI/Private/Input/CommonUIInputSettings.cpp b/Engine/Plugins/Experimental/CommonUI/Source/CommonUI/Private/Input/CommonUIInputSettings.cpp
--- a/Engine/Plugins/Experimental/CommonUI/Source/CommonUI/Private/Input/CommonUIInputSettings.cpp
+++ b/Engine/Plugins/Experimental/CommonUI/Source/CommonUI/Private/Input/CommonUIInputSettings.cpp
@@ -13,16 +13,10 @@ const UCommonUIInputSettings& UCommonUIInputSettings::Get()
 void UCommonUIInputSettings::PostInitProperties()
 {
 	Super::PostInitProperties();
-	
-	if (ensure(HasAnyFlags(RF_ClassDefaultObject)))
-	{
-		//@todo DanH: Resolve the actions and the config overrides
-	}
 }
 
 const FUIInputAction* UCommonUIInputSettings::FindAction(FUIActionTag ActionTag) const
 {
-	//@todo DanH: We'll likely want a TMap for these
 	return InputActions.FindByPredicate([ActionTag](const FUIInputAction& Action) { return Action.ActionTag == ActionTag; });
 }
 
diff --git a/Engine/Plugins/Experimental/CommonUI/Source/CommonUI/Private/Input/UIActionRouterTypes.cpp b/Engine/Plugins/Experimental/CommonUI/Source/CommonUI/Private/Input/UIActionRouterTypes.cpp
--- a/Engine/Plugins/Experimental/CommonUI/Source/CommonUI/Private/Input/UIActionRouterTypes.cpp
+++ b/Engine/Plugins/Experimental/CommonUI/Source/CommonUI/Private/Input/UIActionRouterTypes.cpp
@@ -440,9 +440,6 @@ bool FUIActionBindingHandle::IsValid() const
 	return RegistrationId >= 0 && FUIActionBinding::FindBinding(*this) != nullptr;
 }
 
-//@todo DanH: With widgets caching binding handles that are auto registered/unregistered on construct/destruct, it's less clear what this should be doing
-//		The big question is whether this should fully destroy the binding object, or just unregister it from the node (if it's actually live)
-//		Perhaps we remove the function from here entirely, and you have to pass the handle to the router for unbinding?
 void FUIActionBindingHandle::Unregister()
 {
 	if (TSharedPtr<FUIActionBinding> Binding = FUIActionBinding::FindBinding(*this))
diff --git a/Engine/Plugins/Experimental/CommonUI/Source/CommonUI/Private/SAnalogSlider.cpp b/Engine/Plugins/Experimental/CommonUI/Source/CommonUI/Private/SAnalogSlider.cpp
--- a/Engine/Plugins/Experimental/CommonUI/Source/CommonUI/Private/SAnalogSlider.cpp
+++ b/Engine/Plugins/Experimental/CommonUI/Source/CommonUI/Private/SAnalogSlider.cpp
@@ -129,7 +129,7 @@ FReply SAnalogSlider::OnAnalogValueChanged(const FGeometry& MyGeometry, const FA
 
 FNavigationReply SAnalogSlider::OnNavigation(const FGeometry& MyGeometry, const FNavigationEvent& InNavigationEvent)
 {
-	//TODO Handle up and down.
+	// @TODO: Handle up and down.
 	if (Orientation == EOrientation::Orient_Horizontal)
 	{
 		if (InNavigationEvent.GetNavigationType() == EUINavigation::Left)
diff --git a/Engine/Plugins/Experimental/CommonUI/Source/CommonUI/Private/Widgets/CommonActivatableWidgetContainer.cpp b/Engine/Plugins/Experimental/CommonUI/Source/CommonUI/Private/Widgets/CommonActivatableWidgetContainer.cpp
--- a/Engine/Plugins/Experimental/CommonUI/Source/CommonUI/Private/Widgets/CommonActivatableWidgetContainer.cpp
+++ b/Engine/Plugins/Experimental/CommonUI/Source/CommonUI/Private/Widgets/CommonActivatableWidgetContainer.cpp
@@ -11,7 +11,6 @@
 
 UCommonActivatableWidget* ActivatableWidgetFromSlate(const TSharedPtr<SWidget>& SlateWidget)
 {
-	//@todo DanH: FActivatableWidgetMetaData
 	if (SlateWidget && SlateWidget != SNullWidget::NullWidget && ensure(SlateWidget->GetType().IsEqual(TEXT("SObjectWidget"))))
 	{
 		UCommonActivatableWidget* ActivatableWidget = Cast<UCommonActivatableWidget>(StaticCastSharedPtr<SObjectWidget>(SlateWidget)->GetWidgetObject());
@@ -107,11 +106,6 @@ TSharedRef<SWidget> UCommonActivatableWidgetContainerBase::RebuildWidget()
 
 	// We always want a 0th slot to be able to animate the first real entry in and out
 	MySwitcher->AddSlot() [SNullWidget::NullWidget];
-	
-	if (WidgetList.Num() > 0)
-	{
-		//@todo DanH: We gotta fill up the switcher and activate the right thing. Alternatively we could just flush all the children upon destruction... I think I like that more.
-	}
 
 	return MyOverlay.ToSharedRef();
 }
@@ -180,7 +174,7 @@ UCommonActivatableWidget* UCommonActivatableWidgetContainerBase::AddWidgetIntern
 
 void UCommonActivatableWidgetContainerBase::RegisterInstanceInternal(UCommonActivatableWidget& NewWidget)
 {
-	//@todo DanH: We should do something here to force-disable the bAutoActivate property on the provided widget, since it quite simply makes no sense
+	// @TODO: Log if bAutoActivate is true on the provided widget, since it quite simply makes no sense.
 	if (ensure(!WidgetList.Contains(&NewWidget)))
 	{
 		WidgetList.Add(&NewWidget);
@@ -328,10 +322,6 @@ void UCommonActivatableWidgetStack::OnWidgetAddedToList(UCommonActivatableWidget
 {
 	if (MySwitcher)
 	{
-		//@todo DanH: Rig up something to skip to an index immediately but still play the intro portion of the transition on the new index
-		//		Might be as simple as changing the properties to separate intro and outro durations?
-		//		Eh, but even in this case we only want to skip when we're going from an empty 0th entry. Every other transition should still do the full fade.
-
 		// Toss the widget onto the end of the switcher's children and transition to it immediately
 		MySwitcher->AddSlot() [AddedWidget.TakeWidget()];
 
diff --git a/Engine/Plugins/Experimental/CommonUI/Source/CommonUI/Public/CommonActionWidget.h b/Engine/Plugins/Experimental/CommonUI/Source/CommonUI/Public/CommonActionWidget.h
--- a/Engine/Plugins/Experimental/CommonUI/Source/CommonUI/Public/CommonActionWidget.h
+++ b/Engine/Plugins/Experimental/CommonUI/Source/CommonUI/Public/CommonActionWidget.h
@@ -101,7 +101,6 @@ protected:
 	UPROPERTY(EditAnywhere, BlueprintReadOnly, Category = CommonActionWidget, meta = (EditCondition = "CommonInput.CommonInputSettings.IsEnhancedInputSupportEnabled", EditConditionHides))
 	TObjectPtr<class UInputAction> EnhancedInputAction;
 
-	//@todo DanH: Create clearer split between support for the new & legacy system in here
 	FUIActionBindingHandle DisplayedBindingHandle;
 
 #if WITH_EDITORONLY_DATA
diff --git a/Engine/Plugins/Experimental/CommonUI/Source/CommonUI/Public/CommonActivatableWidget.h b/Engine/Plugins/Experimental/CommonUI/Source/CommonUI/Public/CommonActivatableWidget.h
--- a/Engine/Plugins/Experimental/CommonUI/Source/CommonUI/Public/CommonActivatableWidget.h
+++ b/Engine/Plugins/Experimental/CommonUI/Source/CommonUI/Public/CommonActivatableWidget.h
@@ -27,8 +27,6 @@ DECLARE_DYNAMIC_MULTICAST_DELEGATE(FOnWidgetActivationChanged);
  * 
  * Note that removing an activatable widget from the UI (i.e. triggering Destruct()) will always deactivate it, even if the UWidget is not destroyed.
  * Re-constructing the underlying SWidget will only result in re-activation if auto-activate is enabled.
- *
- * TODO: ADD MORE INFO ON INPUTS
  */
 UCLASS(meta = (DisableNativeTick))
 class COMMONUI_API UCommonActivatableWidget : public UCommonUserWidget
@@ -164,7 +162,6 @@ protected:
 	UPROPERTY(EditAnywhere, Category = Activation)
 	bool bAutoActivate = false;
 
-	//@todo DanH: This property name suuucks, need to circle back and pick something better
 	/**
 	 * True if this widget is a candidate to receive/route focus or specify a desired UIInputConfig when active.
 	 * Primary reason for disabling is for utility sub-widgets within a larger screen that possess actions, but are never
diff --git a/Engine/Plugins/Experimental/CommonUI/Source/CommonUI/Public/CommonUIRichTextData.h b/Engine/Plugins/Experimental/CommonUI/Source/CommonUI/Public/CommonUIRichTextData.h
--- a/Engine/Plugins/Experimental/CommonUI/Source/CommonUI/Public/CommonUIRichTextData.h
+++ b/Engine/Plugins/Experimental/CommonUI/Source/CommonUI/Public/CommonUIRichTextData.h
@@ -21,7 +21,8 @@ struct COMMONUI_API FRichTextIconData : public FTableRowBase
 	FVector2D ImageSize = FVector2D(64.f, 64.f);
 };
 
-/* Derive from this class for rich text data per game
+/** 
+ * Derive from this class for rich text data per game
  * it is referenced in Common UI Settings, found in project settings UI
  */
 UCLASS(Abstract, Blueprintable, meta = (Category = "Common UI"))
@@ -32,7 +33,6 @@ class COMMONUI_API UCommonUIRichTextData : public UObject
 public:
 	static UCommonUIRichTextData* Get();
 
-	//@todo StephanJ: think of methods for game to add more icons to the set
 	const FRichTextIconData* FindIcon(const FName& InKey);
 	const TMap<FName, uint8*>& GetIconMap() const { return InlineIconSet->GetRowMap(); }
 
diff --git a/Engine/Plugins/Experimental/CommonUI/Source/CommonUI/Public/Groups/CommonWidgetGroupBase.h b/Engine/Plugins/Experimental/CommonUI/Source/CommonUI/Public/Groups/CommonWidgetGroupBase.h
--- a/Engine/Plugins/Experimental/CommonUI/Source/CommonUI/Public/Groups/CommonWidgetGroupBase.h
+++ b/Engine/Plugins/Experimental/CommonUI/Source/CommonUI/Public/Groups/CommonWidgetGroupBase.h
@@ -5,8 +5,9 @@
 #include "Components/Widget.h"
 #include "CommonWidgetGroupBase.generated.h"
 
-//@todo DanH: This is only used for buttons, so ditch the base. Also the vast majority of use cases are native, so it also shouldn't be a UObject
-
+/**
+ * Base class for CommonUI widget groups, currently only used for button groups
+ */
 UCLASS(Abstract, BlueprintType)
 class COMMONUI_API UCommonWidgetGroupBase : public UObject
 {
diff --git a/Engine/Plugins/Experimental/CommonUI/Source/CommonUI/Public/Input/CommonAnalogCursor.h b/Engine/Plugins/Experimental/CommonUI/Source/CommonUI/Public/Input/CommonAnalogCursor.h
--- a/Engine/Plugins/Experimental/CommonUI/Source/CommonUI/Public/Input/CommonAnalogCursor.h
+++ b/Engine/Plugins/Experimental/CommonUI/Source/CommonUI/Public/Input/CommonAnalogCursor.h
@@ -83,11 +83,8 @@ protected:
 	
 	bool IsUsingGamepad() const;
 
-
 	bool ShouldHideCursor() const;
 
-
-
 	// Knowingly unorthodox member reference to a UObject - ok because we are a subobject of the owning router and will never outlive it
 	const UCommonUIActionRouterBase& ActionRouter;
 	ECommonInputType ActiveInputMethod;
@@ -97,8 +94,7 @@ protected:
 	bool bShouldHandleRightAnalog = true;
 	
 private:
-	//EAnalogStick GetScrollStick() const;
-	
+
 	/** The current set of pointer buttons being used as keys. */
 	TSet<FKey> PointerButtonDownKeys;
 
@@ -106,10 +102,6 @@ private:
 	FSlateRenderTransform LastCursorTargetTransform;
 
 	float TimeUntilScrollUpdate = 0.f;
-
-	//@todo DanH: This actually needs to go up a layer into FAnalogCursor so it behaves as desired if we allow things to fall through
-	//		We'll also need a CommonNavigationConfig that respects the stick assignment as well so it also maps to nav mode
-	//EAnalogStick CursorMovementStick = EAnalogStick::Left;
 	
 #if !UE_BUILD_SHIPPING
 	enum EShoulderButtonFlags
diff --git a/Engine/Plugins/Experimental/CommonUI/Source/CommonUI/Public/Input/CommonUIActionRouterBase.h b/Engine/Plugins/Experimental/CommonUI/Source/CommonUI/Public/Input/CommonUIActionRouterBase.h
--- a/Engine/Plugins/Experimental/CommonUI/Source/CommonUI/Public/Input/CommonUIActionRouterBase.h
+++ b/Engine/Plugins/Experimental/CommonUI/Source/CommonUI/Public/Input/CommonUIActionRouterBase.h
@@ -51,8 +51,10 @@ enum class ERouteUIInputResult : uint8
 };
 
 /**
- * The nucleus of the CommonUI input routing system
- * @todo DanH: Explain what that means more fully
+ * The nucleus of the CommonUI input routing system. 
+ * 
+ * Gathers input from external sources such as game viewport client and forwards them to widgets 
+ * via activatable tree node representation.
  */
 UCLASS()
 class COMMONUI_API UCommonUIActionRouterBase : public ULocalPlayerSubsystem
@@ -93,7 +95,6 @@ public:
 
 	UCommonInputSubsystem& GetInputSubsystem() const;
 
-	//@todo DanH: VERY TEMP! The event we want from the game viewport lives at the level atm
 	ERouteUIInputResult ProcessInput(FKey Key, EInputEvent InputEvent) const;
 	bool CanProcessNormalGameInput() const;
 
@@ -107,9 +108,7 @@ public:
 
 	void SetActiveUIInputConfig(const FUIInputConfig& NewConfig);
 
-//COMMONUI_SCOPE:
 public:
-	//@todo DanH: Not loving this bit of coupling, it really should to be possible for any widget to accomplish this (and the CommonUserWidget just does it automatically)
 	void NotifyUserWidgetConstructed(const UCommonUserWidget& Widget);
 	void NotifyUserWidgetDestructed(const UCommonUserWidget& Widget);
 	
diff --git a/Engine/Plugins/Experimental/CommonUI/Source/CommonUI/Public/Input/CommonUIInputSettings.h b/Engine/Plugins/Experimental/CommonUI/Source/CommonUI/Public/Input/CommonUIInputSettings.h
--- a/Engine/Plugins/Experimental/CommonUI/Source/CommonUI/Public/Input/CommonUIInputSettings.h
+++ b/Engine/Plugins/Experimental/CommonUI/Source/CommonUI/Public/Input/CommonUIInputSettings.h
@@ -32,17 +32,6 @@ public:
 	/** Time (in seconds) for hold progress to go from 1.0 (completed) to 0.0. */
 	UPROPERTY(EditAnywhere, Config, Category = "UI Action Key Mapping")
 	float HoldRollbackTime = 0.f;
-
-	//@todo DanH: Is this actually wanted/needed at all? Do we really want different bindings for the same thing on Platform A vs Platform B, for instance?
-	///** 
-	// * (Optional) The platforms to which this mapping is exclusive. Leave empty to allow on all platforms.
-	// * Note that most platforms support multiple input methods (Gamepad/KBM/Touch), so be sparing with specific platform restrictions
-	// */
-	//UPROPERTY(EditAnywhere, Config, Category = "UI Action Key Mapping", meta = (Categories = "UI.Platform."))
-	//FGameplayTagContainer Platforms;
-
-	//@todo DanH: Do we care about keyboard modifier keys? Leaning no until a need arises
-	//		If it does, we can just swap out the FKey for an FInputChord and do some PostSerialize fixup
 };
 
 USTRUCT()
@@ -113,8 +102,6 @@ class UCommonUIInputSettings : public UObject
 	GENERATED_BODY()
 
 public:
-	//@todo DanH: Consider collapsing the other settings objects in here too
-
 	COMMONUI_API static const UCommonUIInputSettings& Get();
 	
 	virtual void PostInitProperties() override;
@@ -123,19 +110,12 @@ public:
 	const FUIInputAction* FindAction(FUIActionTag ActionTag) const;
 	const TArray<FUIInputAction>& GetUIInputActions() const { return InputActions; }
 	const FCommonAnalogCursorSettings& GetAnalogCursorSettings() const { return AnalogCursorSettings; }
-
-#if WITH_EDITOR
-	//virtual void PostEditChangeProperty(struct FPropertyChangedEvent& PropertyChangedEvent) override;
-#endif
 //
 private:
 	/** True to have the mouse pointer automatically moved to the center of whatever widget is currently focused while using a gamepad. */
 	UPROPERTY(EditAnywhere, Config, Category = General)
 	bool bLinkCursorToGamepadFocus = true;
 
-	//@todo DanH: I really don't like this, but it's the only means available atm for us to push input components with some level of control
-	//		Short of a larger revamp in PlayerInput, the only other option is to require a manual linkup within the project-level PlayerController's BuildInputStack
-	//		that gets in touch with the action router directly and lets it push components. Opting instead for the decoupled approach, crude as it may be.
 	/** 
 	 * The input priority of the input components that process UI input actions.
 	 * The lower the value, the higher the priority of the component.
@@ -159,12 +139,6 @@ private:
 
 	UPROPERTY(EditAnywhere, Config, Category = AnalogCursor)
 	FCommonAnalogCursorSettings AnalogCursorSettings;
-
-//#if WITH_EDITORONLY_DATA
-//	/** Purely a visual reference to reflect the actual priority values for the input components for each mode, based UIActionProcessingPriority. */
-//	UPROPERTY(VisibleAnywhere, Category = InputSettings)
-//	int32 InputModePriorities[(uint8)ECommonInputMode::MAX];
-//#endif
 };
 
 #if UE_ENABLE_INCLUDE_ORDER_DEPRECATED_IN_5_2
diff --git a/Engine/Plugins/Experimental/CommonUI/Source/CommonUI/Public/Widgets/CommonActivatableWidgetContainer.h b/Engine/Plugins/Experimental/CommonUI/Source/CommonUI/Public/Widgets/CommonActivatableWidgetContainer.h
--- a/Engine/Plugins/Experimental/CommonUI/Source/CommonUI/Public/Widgets/CommonActivatableWidgetContainer.h
+++ b/Engine/Plugins/Experimental/CommonUI/Source/CommonUI/Public/Widgets/CommonActivatableWidgetContainer.h
@@ -94,7 +94,6 @@ protected:
 
 	virtual void OnWidgetAddedToList(UCommonActivatableWidget& AddedWidget) { unimplemented(); }
 
-	//@todo DanH: There's enough here that reuse is good, but the name of the function undermines how specifically it's meant to be used
 	void SetSwitcherIndex(int32 TargetIndex, bool bInstantTransition = false);
 
 	/** The type of transition to play between widgets */
@@ -161,7 +160,6 @@ private:
 	void ReleaseWidget(const TSharedRef<SWidget>& WidgetToRelease);
 	TArray<TSharedPtr<SWidget>> ReleasedWidgets;
 
-	//@todo DanH: This is pretty clunky - revisit for something smoother
 	bool bRemoveDisplayedWidgetPostTransition = false;
 
 	mutable FOnDisplayedWidgetChanged OnDisplayedWidgetChangedEvent;
@@ -184,7 +182,7 @@ class COMMONUI_API UCommonActivatableWidgetStack : public UCommonActivatableWidg
 	GENERATED_BODY()
 
 public:
-	//@todo DanH: Allow the root content class to be restricted by metadata
+
 	UCommonActivatableWidget* GetRootContent() const;
 
 protected:
diff --git a/Engine/Plugins/Experimental/CommonUI/Source/CommonUIEditor/Private/CommonAssetTypeActions_GenericInputActionDataTable.cpp b/Engine/Plugins/Experimental/CommonUI/Source/CommonUIEditor/Private/CommonAssetTypeActions_GenericInputActionDataTable.cpp
--- a/Engine/Plugins/Experimental/CommonUI/Source/CommonUIEditor/Private/CommonAssetTypeActions_GenericInputActionDataTable.cpp
+++ b/Engine/Plugins/Experimental/CommonUI/Source/CommonUIEditor/Private/CommonAssetTypeActions_GenericInputActionDataTable.cpp
@@ -25,7 +25,6 @@ UClass* FCommonAssetTypeActions_GenericInputActionDataTable::GetSupportedClass()
 
 void FCommonAssetTypeActions_GenericInputActionDataTable::OpenAssetEditor(const TArray<UObject*>& InObjects, TSharedPtr<IToolkitHost> EditWithinLevelEditor)
 {
-	//@TODO: DarenC - This is from FAssetTypeActions_DataTable - but we can't derive to private include.
 	TArray<UDataTable*> DataTablesToOpen;
 	TArray<UDataTable*> InvalidDataTables;
 
