Add support for specifying Input Configs from BP for activatable widgets
Adding various debug CVars to CommonUI

Add CVar to warn when activating widget without GetDesiredFocusTarget implemented
Add CVar to log Input Config Canges.
Add CVar to warn when using CommonUI without CommonGameViewportClient
Add CVar to warn when deactivating all widgets.
Add Slate CVar to log when navigation config changes

diff --git a/Engine/Plugins/Experimental/CommonUI/Source/CommonUI/Private/CommonActivatableWidget.cpp b/Engine/Plugins/Experimental/CommonUI/Source/CommonUI/Private/CommonActivatableWidget.cpp
--- a/Engine/Plugins/Experimental/CommonUI/Source/CommonUI/Private/CommonActivatableWidget.cpp
+++ b/Engine/Plugins/Experimental/CommonUI/Source/CommonUI/Private/CommonActivatableWidget.cpp
@@ -15,6 +15,16 @@
 
 #include UE_INLINE_GENERATED_CPP_BY_NAME(CommonActivatableWidget)
 
+//////////////////////////////////////////////////////////////////////////
+
+static bool bWarnDesiredFocusNotImplemented = false;
+static FAutoConsoleVariableRef CVarWarnDesiredFocusNotImplemented(
+	TEXT("CommonUI.Debug.WarnDesiredFocusNotImplemented"),
+	bWarnDesiredFocusNotImplemented,
+	TEXT("Warn when a widget is activated without GetDesiredFocus implemented."));
+
+//////////////////////////////////////////////////////////////////////////
+
 UCommonActivatableWidget::FActivatableWidgetRebuildEvent UCommonActivatableWidget::OnRebuilding;
 
 void UCommonActivatableWidget::NativeConstruct()
@@ -74,11 +84,22 @@ TOptional<FActivationMetadata> UCommonActivatableWidget::GetActivationMetadata()
 
 UWidget* UCommonActivatableWidget::NativeGetDesiredFocusTarget() const
 {
+	if (bWarnDesiredFocusNotImplemented && !GetClass()->IsFunctionImplementedInScript(GET_FUNCTION_NAME_CHECKED(UCommonActivatableWidget, BP_GetDesiredFocusTarget)))
+	{
+		UE_LOG(LogCommonUI, Warning, TEXT("[%s] -> Desired Focus not implemented!"), *GetName());
+	}
+
 	return BP_GetDesiredFocusTarget();
 }
 
 TOptional<FUIInputConfig> UCommonActivatableWidget::GetDesiredInputConfig() const
 {
+	// Check if there is a BP implementation for input configs
+	if (GetClass()->IsFunctionImplementedInScript(GET_FUNCTION_NAME_CHECKED(UCommonActivatableWidget, BP_GetDesiredInputConfig)))
+	{
+		return BP_GetDesiredInputConfig();
+	}
+
 	// No particular config is desired by default
 	return TOptional<FUIInputConfig>();
 }
diff --git a/Engine/Plugins/Experimental/CommonUI/Source/CommonUI/Private/Input/CommonUIActionRouterBase.cpp b/Engine/Plugins/Experimental/CommonUI/Source/CommonUI/Private/Input/CommonUIActionRouterBase.cpp
--- a/Engine/Plugins/Experimental/CommonUI/Source/CommonUI/Private/Input/CommonUIActionRouterBase.cpp
+++ b/Engine/Plugins/Experimental/CommonUI/Source/CommonUI/Private/Input/CommonUIActionRouterBase.cpp
@@ -26,12 +26,49 @@
 
 #include UE_INLINE_GENERATED_CPP_BY_NAME(CommonUIActionRouterBase)
 
+//////////////////////////////////////////////////////////////////////////
+
 bool bAlwaysShowCursor = false;
 static const FAutoConsoleVariableRef CVarAlwaysShowCursor(
 	TEXT("CommonUI.AlwaysShowCursor"),
 	bAlwaysShowCursor,
 	TEXT(""));
 
+//////////////////////////////////////////////////////////////////////////
+
+static bool bTraceInputConfig = false;
+static FAutoConsoleVariableRef CVarTraceInputConfig(
+	TEXT("CommonUI.Debug.TraceConfigChanges"),
+	bTraceInputConfig,
+	TEXT("Trace Input Config transitions (Non-shipping). Suggest use with Slate.Debug.TraceNavigationConfig (Non-shipping)."));
+
+static bool bTraceConfigOnScreen = false;
+static FAutoConsoleVariableRef CVarTraceConfigOnScreen(
+	TEXT("CommonUI.Debug.TraceConfigOnScreen"),
+	bTraceConfigOnScreen,
+	TEXT("Trace for input configs should be displayed on screen. Requires CommonUI.Debug.TraceConfigChanges."));
+
+static int32 TraceInputConfigNum = 5;
+static FAutoConsoleVariableRef CVarTraceInputConfigNum(
+	TEXT("CommonUI.Debug.TraceInputConfigNum"),
+	TraceInputConfigNum,
+	TEXT("Number of Input config to keep in trace history."),
+	ECVF_ReadOnly);
+
+static bool bWarnAllWidgetsDeactivated = false;
+static FAutoConsoleVariableRef CVarWarnAllWidgetsDeactivated(
+	TEXT("CommonUI.Debug.WarnAllWidgetsDeactivated"),
+	bWarnAllWidgetsDeactivated,
+	TEXT("Warn when all widgets are deactivated. A valid event, but may have leftover input configs."));
+
+static bool bCheckGameViewportClientValid = true;
+static FAutoConsoleVariableRef CVarCheckGameViewportClientValid(
+	TEXT("CommonUI.Debug.CheckGameViewportClientValid"),
+	bCheckGameViewportClientValid,
+	TEXT("Log error when CommonUI is used without the current game viewport deriving from CommonGameViewportClient."));
+
+//////////////////////////////////////////////////////////////////////////
+
 FGlobalUITags FGlobalUITags::GUITags;
 
 UCommonActivatableWidget* FindOwningActivatable(const UWidget& Widget)
@@ -102,6 +139,7 @@ UCommonActivatableWidget* UCommonUIActionRouterBase::FindOwningActivatable(TShar
 
 UCommonUIActionRouterBase::UCommonUIActionRouterBase()
 	: PersistentActions(MakeShared<FPersistentActionCollection>(*this))
+	, InputConfigSources(TraceInputConfigNum, "None")
 {
 	// Non-CDO behavior
 	if (!HasAnyFlags(RF_ClassDefaultObject))
@@ -197,6 +235,13 @@ void UCommonUIActionRouterBase::Initialize(FSubsystemCollectionBase& Collection)
 	{
 		AnalogCursor = MakeAnalogCursor();
 		PostAnalogCursorCreate();
+
+		if (bCheckGameViewportClientValid && !GEngine->GameViewportClientClass->IsChildOf<UCommonGameViewportClient>())
+		{
+			UE_LOG(LogUIActionRouter, Error, 
+				TEXT("Using CommonUI without a CommonGameViewportClient derived game viewport client. CommonUI Input routing will not function correctly.\n")
+				TEXT("To disable this warning set CommonUI.Debug.CheckGameViewportClientValid=0 under [SystemSettings] in your project's DefaultEngine.ini."));
+		}
 	}
 	else
 	{
@@ -678,6 +723,12 @@ void UCommonUIActionRouterBase::HandleRootNodeDeactivated(TWeakPtr<FActivatableT
 	// In the case that all root nodes are not activated we need to re-establish input for the highest paint layer node in action domain nodes.
 	if (!bActivatedRootNodeExists && bIsActivatableTreeEnabled)
 	{
+		if (bWarnAllWidgetsDeactivated)
+		{
+			UE_LOG(LogUIActionRouter, Warning, TEXT("All widgets deactivated. Existing input config set: %s"), 
+				ActiveInputConfig.IsSet() ? TEXT("Yes - the current input config is lingering from a deactivated widget.") : TEXT("No."));
+		}
+
 		RefreshActionDomainLeafNodeConfig();
 	}
 }
@@ -1279,8 +1330,34 @@ FActivatableTreeNodePtr UCommonUIActionRouterBase::FindNodeRecursive(const FActi
 	return FoundNode;
 }
 
-void UCommonUIActionRouterBase::SetActiveUIInputConfig(const FUIInputConfig& NewConfig)
+void UCommonUIActionRouterBase::SetActiveUIInputConfig(const FUIInputConfig& NewConfig, const UObject* InConfigSource)
 {
+#if WITH_SLATE_DEBUGGING
+	if (bTraceInputConfig && InConfigSource && (!ActiveInputConfig.IsSet() || NewConfig != ActiveInputConfig.GetValue()))
+	{
+		InputConfigSources[InputConfigSourceIndex] = InConfigSource->GetName();
+
+		TStringBuilder<1024> Builder;
+		Builder.Appendf(TEXT("Input Config Change:\n%s"), *NewConfig.ToString());
+		Builder.Appendf(TEXT("--- Input Config Source History (Newest Last) ---"));
+		for (uint32 Index = 0; Index < InputConfigSources.Capacity(); Index++)
+		{
+			uint32 CircularIndex = InputConfigSources.GetNextIndex(Index + InputConfigSourceIndex);
+			Builder.Appendf(TEXT("Config Source: %s"), *InputConfigSources[CircularIndex]);
+		}
+		FString InputConfigTrace = Builder.ToString();
+
+		UE_LOG(LogUIActionRouter, Log, TEXT("%s"), *InputConfigTrace);
+		if (bTraceConfigOnScreen)
+		{
+			uint64 OnScreenTraceKey = 202013;
+			GEngine->AddOnScreenDebugMessage(OnScreenTraceKey, 15.f, FColor::Cyan, InputConfigTrace);
+		}
+
+		InputConfigSourceIndex = InputConfigSources.GetNextIndex(InputConfigSourceIndex);
+	}
+#endif // WITH_SLATE_DEBUGGING
+
 	ApplyUIInputConfig(NewConfig, !ActiveInputConfig.IsSet());
 }
 
@@ -1310,7 +1387,7 @@ void UCommonUIActionRouterBase::RefreshActionDomainLeafNodeConfig()
 				}
 			}
 
-			SetActiveUIInputConfig(FUIInputConfig(ActionDomainTable->InputMode, ActionDomainTable->MouseCaptureMode));
+			SetActiveUIInputConfig(FUIInputConfig(ActionDomainTable->InputMode, ActionDomainTable->MouseCaptureMode), ActionDomainTable);
 		}
 	}
 }
diff --git a/Engine/Plugins/Experimental/CommonUI/Source/CommonUI/Private/Input/UIActionRouterTypes.cpp b/Engine/Plugins/Experimental/CommonUI/Source/CommonUI/Private/Input/UIActionRouterTypes.cpp
--- a/Engine/Plugins/Experimental/CommonUI/Source/CommonUI/Private/Input/UIActionRouterTypes.cpp
+++ b/Engine/Plugins/Experimental/CommonUI/Source/CommonUI/Private/Input/UIActionRouterTypes.cpp
@@ -543,6 +543,19 @@ FUIInputConfig::FUIInputConfig()
 	, MouseCaptureMode(EMouseCaptureMode::NoCapture)
 {}
 
+FString FUIInputConfig::ToString() const
+{
+	TStringBuilder<1024> Builder;
+
+	Builder.Appendf(TEXT("bIgnoreMoveInput: %u\n"), bIgnoreMoveInput);
+	Builder.Appendf(TEXT("bIgnoreLookInput: %u\n"), bIgnoreLookInput);
+	Builder.Appendf(TEXT("InputMode: %s\n"), *UEnum::GetValueAsString(InputMode));
+	Builder.Appendf(TEXT("MouseCaptureMode: %s\n"), *UEnum::GetValueAsString(MouseCaptureMode));
+	Builder.Appendf(TEXT("bHideCursorDuringViewportCapture: %u\n"), bHideCursorDuringViewportCapture);
+
+	return Builder.ToString();
+}
+
 //////////////////////////////////////////////////////////////////////////
 // FActionRouterBindingCollection
 //////////////////////////////////////////////////////////////////////////
@@ -1394,7 +1407,7 @@ void FActivatableTreeRoot::ApplyLeafmostNodeConfig()
 			TOptional<FUIInputConfig> DesiredConfig = PinnedLeafmostNode->FindDesiredInputConfig();
 			if(DesiredConfig.IsSet())
 			{
-				GetActionRouter().SetActiveUIInputConfig(DesiredConfig.GetValue());
+				GetActionRouter().SetActiveUIInputConfig(DesiredConfig.GetValue(), PinnedLeafmostNode->GetWidget());
 			}
 			else if(ICommonInputModule::GetSettings().GetEnableDefaultInputConfig())
 			{
diff --git a/Engine/Plugins/Experimental/CommonUI/Source/CommonUI/Public/CommonActivatableWidget.h b/Engine/Plugins/Experimental/CommonUI/Source/CommonUI/Public/CommonActivatableWidget.h
--- a/Engine/Plugins/Experimental/CommonUI/Source/CommonUI/Public/CommonActivatableWidget.h
+++ b/Engine/Plugins/Experimental/CommonUI/Source/CommonUI/Public/CommonActivatableWidget.h
@@ -122,6 +122,14 @@ protected:
 	 */
 	UFUNCTION(BlueprintImplementableEvent, Category = ActivatableWidget, meta = (DisplayName = "Get Desired Focus Target"))
 	UWidget* BP_GetDesiredFocusTarget() const;
+	
+	/** 
+	 * Implement to provide the input config to use when this widget is activated. Keep in mind when all widgets
+	 * are deactivated, CommonUI will not attempt to automatically restore the input config to before any widget was active.
+	 * Note: This is a fallback used only if the native class parentage does not provide an input config.
+	 */
+	UFUNCTION(BlueprintImplementableEvent, Category = ActivatableWidget, meta = (DisplayName = "Get Desired Input Config"))
+	FUIInputConfig BP_GetDesiredInputConfig() const;
 
 	/**
 	 * Ask for focus to be re-set to our current DesiredFocusTarget,
diff --git a/Engine/Plugins/Experimental/CommonUI/Source/CommonUI/Public/Input/CommonUIActionRouterBase.h b/Engine/Plugins/Experimental/CommonUI/Source/CommonUI/Public/Input/CommonUIActionRouterBase.h
--- a/Engine/Plugins/Experimental/CommonUI/Source/CommonUI/Public/Input/CommonUIActionRouterBase.h
+++ b/Engine/Plugins/Experimental/CommonUI/Source/CommonUI/Public/Input/CommonUIActionRouterBase.h
@@ -5,6 +5,7 @@
 #include "CommonInputModeTypes.h"
 #include "Subsystems/LocalPlayerSubsystem.h"
 #include "Containers/Ticker.h"
+#include "Containers/CircularBuffer.h"
 
 #include "Engine/EngineBaseTypes.h"
 #include "Input/UIActionBindingHandle.h"
@@ -106,7 +107,13 @@ public:
 
 	bool IsWidgetInActiveRoot(const UCommonActivatableWidget* Widget) const;
 
-	void SetActiveUIInputConfig(const FUIInputConfig& NewConfig);
+	/** 
+	 * Sets Input Config 
+	 * 
+	 * @param NewConfig config to set
+	 * @param InConfigSource optional source of config. If exists, will be used to log input config source
+	 */
+	void SetActiveUIInputConfig(const FUIInputConfig& NewConfig, const UObject* InConfigSource = nullptr);
 
 public:
 	void NotifyUserWidgetConstructed(const UCommonUserWidget& Widget);
@@ -201,6 +208,9 @@ private:
 	// Note: Treat this as a TSharedRef - only reason it isn't is because TSharedRef doesn't play nice with forward declarations :(
 	TSharedPtr<class FPersistentActionCollection> PersistentActions;
 
+	TCircularBuffer<FString> InputConfigSources = TCircularBuffer<FString>(5, "None");
+	int32 InputConfigSourceIndex = 0;
+
 	bool bForceResetActiveRoot = false;
 
 	mutable FSimpleMulticastDelegate OnBoundActionsUpdatedEvent;
diff --git a/Engine/Plugins/Experimental/CommonUI/Source/CommonUI/Public/Input/UIActionBindingHandle.h b/Engine/Plugins/Experimental/CommonUI/Source/CommonUI/Public/Input/UIActionBindingHandle.h
--- a/Engine/Plugins/Experimental/CommonUI/Source/CommonUI/Public/Input/UIActionBindingHandle.h
+++ b/Engine/Plugins/Experimental/CommonUI/Source/CommonUI/Public/Input/UIActionBindingHandle.h
@@ -4,6 +4,7 @@
 
 #include "Containers/UnrealString.h"
 #include "Misc/Optional.h"
+#include "UIActionBindingHandle.generated.h"
 
 enum class EMouseCaptureMode : uint8;
 struct FScriptContainerElement;
@@ -76,13 +77,16 @@ private:
  * Input Config that can be applied on widget activation. Allows for input setup  (Mouse capture, 
  * UI-only input, move / look ignore, etc), to be controlled by widget activation.
  */
-struct FUIInputConfig
+USTRUCT(BlueprintType)
+struct COMMONUI_API FUIInputConfig
 {
+	GENERATED_BODY()
+
 	ECommonInputMode GetInputMode() const { return InputMode; }
 	EMouseCaptureMode GetMouseCaptureMode() const { return MouseCaptureMode; }
 	bool HideCursorDuringViewportCapture() const { return bHideCursorDuringViewportCapture; }
 
-	COMMONUI_API FUIInputConfig();
+	FUIInputConfig();
 	FUIInputConfig(ECommonInputMode InInputMode, EMouseCaptureMode InMouseCaptureMode, bool bInHideCursorDuringViewportCapture = true)
 		: InputMode(InInputMode)
 		, MouseCaptureMode(InMouseCaptureMode)
@@ -101,12 +105,25 @@ struct FUIInputConfig
 		return !operator==(Other);
 	}
 
+	UPROPERTY(EditAnywhere, BlueprintReadWrite, Category = InputConfig)
 	bool bIgnoreMoveInput = false;
+
+	UPROPERTY(EditAnywhere, BlueprintReadWrite, Category = InputConfig)
 	bool bIgnoreLookInput = false;
 
-private:
+	/** Simplification of config as string */
+	FString ToString() const;
+
+protected:
+
+
+	UPROPERTY(EditAnywhere, BlueprintReadWrite, Category = InputConfig)
 	ECommonInputMode InputMode;
+
+	UPROPERTY(EditAnywhere, BlueprintReadWrite, Category = InputConfig)
 	EMouseCaptureMode MouseCaptureMode;
+
+	UPROPERTY(EditAnywhere, BlueprintReadWrite, Category = InputConfig)
 	bool bHideCursorDuringViewportCapture = true;
 };
 
diff --git a/Engine/Source/Runtime/Slate/Private/Framework/Application/NavigationConfig.cpp b/Engine/Source/Runtime/Slate/Private/Framework/Application/NavigationConfig.cpp
--- a/Engine/Source/Runtime/Slate/Private/Framework/Application/NavigationConfig.cpp
+++ b/Engine/Source/Runtime/Slate/Private/Framework/Application/NavigationConfig.cpp
@@ -174,3 +174,24 @@ EUINavigationAction FNavigationConfig::GetNavigationActionForKey(const FKey& InK
 
 	return EUINavigationAction::Invalid;
 }
+
+FString FNavigationConfig::ToString() const
+{
+	TStringBuilder<1024> Builder;
+
+	Builder.Appendf(TEXT("bTabNavigation: %u\n"), bTabNavigation);
+	Builder.Appendf(TEXT("bKeyNavigation: %u\n"), bKeyNavigation);
+	Builder.Appendf(TEXT("bAnalogNavigation: %u\n"), bAnalogNavigation);
+	Builder.Appendf(TEXT("AnalogNavigationHorizontalThreshold: %f\n"), AnalogNavigationHorizontalThreshold);
+	Builder.Appendf(TEXT("AnalogNavigationVerticalThreshold: %f\n"), AnalogNavigationVerticalThreshold);
+	Builder.Appendf(TEXT("AnalogHorizontalKey: %s\n"), *AnalogHorizontalKey.ToString());
+	Builder.Appendf(TEXT("AnalogVerticalKey: %s\n"), *AnalogVerticalKey.ToString());
+
+	Builder.Appendf(TEXT("KeyEventRules: \n"));
+	for (TPair<FKey, EUINavigation> KeyEventRule : KeyEventRules)
+	{
+		Builder.Appendf(TEXT("Pair: %s -> %s\n"), *KeyEventRule.Key.ToString(), *UEnum::GetValueAsString(KeyEventRule.Value));
+	}
+
+	return Builder.ToString();
+}
diff --git a/Engine/Source/Runtime/Slate/Private/Framework/Application/SlateApplication.cpp b/Engine/Source/Runtime/Slate/Private/Framework/Application/SlateApplication.cpp
--- a/Engine/Source/Runtime/Slate/Private/Framework/Application/SlateApplication.cpp
+++ b/Engine/Source/Runtime/Slate/Private/Framework/Application/SlateApplication.cpp
@@ -26,6 +26,7 @@
 #include "Input/Events.h"
 #include "Input/HittestGrid.h"
 #include "HAL/PlatformApplicationMisc.h"
+#include "HAL/PlatformStackWalk.h"
 #include "Null/NullPlatformApplicationMisc.h"
 #include "GenericPlatform/GenericPlatformInputDeviceMapper.h"
 
@@ -85,6 +86,13 @@ namespace UE::Slate::Private
 	void VerifyParentChildrenRelationship(const TSharedRef<SWindow>& WindowToDraw);
 	void VerifyWidgetLayerId(const TSharedRef<SWindow>& WindowToDraw);
 }
+
+bool GSlateTraceNavigationConfig = false;
+static FAutoConsoleVariableRef CVarSlateTraceNavigationConfig(
+	TEXT("Slate.Debug.TraceNavigationConfig"),
+	GSlateTraceNavigationConfig,
+	TEXT("True enables tracing of navigation config & callstack to log.")
+);
 #endif //WITH_SLATE_DEBUGGING
 
 bool GSlateInputMotionFiresUserInteractionEvents = true;
@@ -2532,6 +2540,24 @@ bool FSlateApplication::GetTransformFullscreenMouseInput() const
 	return TransformFullscreenMouseInput;
 }
 
+#if WITH_SLATE_DEBUGGING
+void FSlateApplication::TryDumpNavigationConfig(TSharedPtr<FNavigationConfig> InNavigationConfig) const
+{
+	if (GSlateTraceNavigationConfig && InNavigationConfig)
+	{
+		UE_LOG(LogSlate, Log, TEXT("Navigation Config Change:\n%s"), *InNavigationConfig->ToString());
+
+		const uint32 DumpCallstackSize = 65535;
+		ANSICHAR DumpCallstack[DumpCallstackSize] = { 0 };
+		FString ScriptStack = FFrame::GetScriptCallstack(true /* bReturnEmpty */);
+		FPlatformStackWalk::StackWalkAndDump(DumpCallstack, DumpCallstackSize, 0);
+		UE_LOG(LogSlate, Log, TEXT("--- Navigation Config Changing Callstack ---"));
+		UE_LOG(LogSlate, Log, TEXT("Script Stack:\n%s"), *ScriptStack);
+		UE_LOG(LogSlate, Log, TEXT("Callstack:\n%s"), ANSI_TO_TCHAR(DumpCallstack));
+	}
+}
+#endif // WITH_SLATE_DEBUGGING
+
 bool FSlateApplication::SetUserFocus(uint32 UserIndex, const TSharedPtr<SWidget>& WidgetToFocus, EFocusCause ReasonFocusIsChanging /* = EFocusCause::SetDirectly*/)
 {
 	TSharedPtr<FSlateUser> CurrentUser = GetUser(UserIndex);
@@ -6769,6 +6795,10 @@ void FSlateApplication::SetNavigationConfig(TSharedRef<FNavigationConfig> InNavi
 	NavigationConfig->OnUnregister();
 	NavigationConfig = InNavigationConfig;
 	NavigationConfig->OnRegister();
+
+#if WITH_SLATE_DEBUGGING
+	TryDumpNavigationConfig(NavigationConfig);
+#endif // WITH_SLATE_DEBUGGING
 }
 
 bool FSlateApplication::OnConvertibleLaptopModeChanged()
diff --git a/Engine/Source/Runtime/Slate/Private/Framework/Application/SlateUser.cpp b/Engine/Source/Runtime/Slate/Private/Framework/Application/SlateUser.cpp
--- a/Engine/Source/Runtime/Slate/Private/Framework/Application/SlateUser.cpp
+++ b/Engine/Source/Runtime/Slate/Private/Framework/Application/SlateUser.cpp
@@ -582,6 +582,10 @@ void FSlateUser::SetUserNavigationConfig(TSharedPtr<FNavigationConfig> InNavigat
 	{
 		InNavigationConfig->OnRegister();
 	}
+
+#if WITH_SLATE_DEBUGGING
+	FSlateApplication::Get().TryDumpNavigationConfig(UserNavigationConfig);
+#endif // WITH_SLATE_DEBUGGING
 }
 
 bool FSlateUser::IsTouchPointerActive(int32 TouchPointerIndex) const
diff --git a/Engine/Source/Runtime/Slate/Public/Framework/Application/NavigationConfig.h b/Engine/Source/Runtime/Slate/Public/Framework/Application/NavigationConfig.h
--- a/Engine/Source/Runtime/Slate/Public/Framework/Application/NavigationConfig.h
+++ b/Engine/Source/Runtime/Slate/Public/Framework/Application/NavigationConfig.h
@@ -86,6 +86,9 @@ public:
 	UE_DEPRECATED(4.24, "GetNavigationActionForKey doesn't handle multiple users properly, use GetNavigationActionFromKey instead")
 	virtual EUINavigationAction GetNavigationActionForKey(const FKey& InKey) const;
 
+	/** Simplification of config as string */
+	virtual FString ToString() const;
+
 public:
 	/** Should the Tab key perform next and previous style navigation. */
 	bool bTabNavigation;
diff --git a/Engine/Source/Runtime/Slate/Public/Framework/Application/SlateApplication.h b/Engine/Source/Runtime/Slate/Public/Framework/Application/SlateApplication.h
--- a/Engine/Source/Runtime/Slate/Public/Framework/Application/SlateApplication.h
+++ b/Engine/Source/Runtime/Slate/Public/Framework/Application/SlateApplication.h
@@ -583,7 +583,17 @@ public:
 	/**
 	 * True if transforming mouse input coordinates to account for fullscreen distortions
 	 */
-	bool GetTransformFullscreenMouseInput() const; 
+	bool GetTransformFullscreenMouseInput() const;
+
+#if WITH_SLATE_DEBUGGING
+	/**
+	 * Tries to dumps the current navigation config along with callstack used to set it.
+	 * Does nothing if Slate.Debug.TraceNavigationConfig is false (By default is false)
+	 *
+	 * @param InNavigationConfig Navigation config to dump info to log for
+	 */
+	void TryDumpNavigationConfig(TSharedPtr<FNavigationConfig> InNavigationConfig) const;
+#endif // WITH_SLATE_DEBUGGING
 
 	/**
 	 * Sets specified user focus to the SWidget passed in.
