Fix line spacing in CommonUI Input Config trace logs

diff --git a/Engine/Plugins/Experimental/CommonUI/Source/CommonUI/Private/Input/CommonUIActionRouterBase.cpp b/Engine/Plugins/Experimental/CommonUI/Source/CommonUI/Private/Input/CommonUIActionRouterBase.cpp
--- a/Engine/Plugins/Experimental/CommonUI/Source/CommonUI/Private/Input/CommonUIActionRouterBase.cpp
+++ b/Engine/Plugins/Experimental/CommonUI/Source/CommonUI/Private/Input/CommonUIActionRouterBase.cpp
@@ -1338,12 +1338,12 @@ void UCommonUIActionRouterBase::SetActiveUIInputConfig(const FUIInputConfig& New
 		InputConfigSources[InputConfigSourceIndex] = InConfigSource->GetName();
 
 		TStringBuilder<1024> Builder;
-		Builder.Appendf(TEXT("Input Config Change:\n%s"), *NewConfig.ToString());
-		Builder.Appendf(TEXT("--- Input Config Source History (Newest Last) ---"));
+		Builder.Appendf(TEXT("Input Config Change:\n%s\n"), *NewConfig.ToString());
+		Builder.Appendf(TEXT("--- Input Config Source History (Newest Last) ---\n"));
 		for (uint32 Index = 0; Index < InputConfigSources.Capacity(); Index++)
 		{
 			uint32 CircularIndex = InputConfigSources.GetNextIndex(Index + InputConfigSourceIndex);
-			Builder.Appendf(TEXT("Config Source: %s"), *InputConfigSources[CircularIndex]);
+			Builder.Appendf(TEXT("%s\n"), *InputConfigSources[CircularIndex]);
 		}
 		FString InputConfigTrace = Builder.ToString();
 
