UI - Adding a togglable property for the UCommonButtonBase class for hovering effect with touch input while keeping the cvar

diff --git a/Engine/Plugins/Experimental/CommonUI/Source/CommonUI/Private/CommonButtonBase.cpp b/Engine/Plugins/Experimental/CommonUI/Source/CommonUI/Private/CommonButtonBase.cpp
--- a/Engine/Plugins/Experimental/CommonUI/Source/CommonUI/Private/CommonButtonBase.cpp
+++ b/Engine/Plugins/Experimental/CommonUI/Source/CommonUI/Private/CommonButtonBase.cpp
@@ -32,11 +32,11 @@
 
 namespace UE::CommonUI::Private
 {
-	int32 bSimulateHoverOnTouchInput = 1;
-	FAutoConsoleVariableRef CVarSimulateHoverOnTouchInput(
-		TEXT("CommonButton.SimulateHoverForTouch"),
-		bSimulateHoverOnTouchInput,
-		TEXT("Enables or disables the hovering/unhovering simulation on Touch Input System.\n0: Disable, 1: Enable (default)"),
+	int32 bEnableSimulateHoverOnTouchInput = 1;
+	FAutoConsoleVariableRef CVarEnableSimulateHoverOnTouchInput(
+		TEXT("CommonButton.EnableSimulateHoverOnTouch"),
+		bEnableSimulateHoverOnTouchInput,
+		TEXT("Allows buttons to simulate hovering on touch in accordance with the property SimulateHoverOnTouchInput.\n0: Disable, 1: Enable (default)"),
 		ECVF_Default);
 }
 
@@ -349,6 +349,7 @@ UCommonButtonBase::UCommonButtonBase(const FObjectInitializer& ObjectInitializer
 	, bDisplayInputActionWhenNotInteractable(true)
 	, bShouldUseFallbackDefaultInputAction(true)
 	, bRequiresHold(false)
+	, bSimulateHoverOnTouchInput(true)
 	, bSelected(false)
 	, bButtonEnabled(true)
 	, bInteractionEnabled(true)
@@ -1344,7 +1345,7 @@ void UCommonButtonBase::HandleButtonPressed()
 
 	UCommonInputSubsystem* CommonInputSubsystem = GetInputSubsystem();
 
-	if (CommonInputSubsystem && CommonInputSubsystem->GetCurrentInputType() == ECommonInputType::Touch && UE::CommonUI::Private::bSimulateHoverOnTouchInput)
+	if (CommonInputSubsystem && CommonInputSubsystem->GetCurrentInputType() == ECommonInputType::Touch && bSimulateHoverOnTouchInput && UE::CommonUI::Private::bEnableSimulateHoverOnTouchInput)
 	{
 		// Simulate hover events when using touch input
 		NativeOnHovered();
@@ -1372,7 +1373,7 @@ void UCommonButtonBase::HandleButtonReleased()
 
 	UCommonInputSubsystem* CommonInputSubsystem = GetInputSubsystem();
 
-	if (CommonInputSubsystem && CommonInputSubsystem->GetCurrentInputType() == ECommonInputType::Touch && UE::CommonUI::Private::bSimulateHoverOnTouchInput)
+	if (CommonInputSubsystem && CommonInputSubsystem->GetCurrentInputType() == ECommonInputType::Touch && bSimulateHoverOnTouchInput && UE::CommonUI::Private::bEnableSimulateHoverOnTouchInput)
 	{
 		// Simulate hover events when using touch input
 		NativeOnUnhovered();
diff --git a/Engine/Plugins/Experimental/CommonUI/Source/CommonUI/Public/CommonButtonBase.h b/Engine/Plugins/Experimental/CommonUI/Source/CommonUI/Public/CommonButtonBase.h
--- a/Engine/Plugins/Experimental/CommonUI/Source/CommonUI/Public/CommonButtonBase.h
+++ b/Engine/Plugins/Experimental/CommonUI/Source/CommonUI/Public/CommonButtonBase.h
@@ -773,6 +773,10 @@ protected:
 	/** Press and Hold values used for Keyboard and Mouse, Gamepad and Touch, depending on the current input type */
 	UPROPERTY(EditAnywhere, BlueprintReadOnly, Category = "Input|Hold", meta = (EditCondition="bRequiresHold", ExposeOnSpawn = true))
 	TSubclassOf<UCommonUIHoldData> HoldData;
+	
+	/** True if this button should play the hover effect when pressed by a touch input */
+	UPROPERTY(EditAnywhere, BlueprintReadOnly, Category = Input, AdvancedDisplay)
+	bool bSimulateHoverOnTouchInput = true;
 
 private:
 
