PR #10386: Fix widgets not being able to register multiple enhanced input actions (Contributed by [REDACTED_PERSON])

diff --git a/Engine/Plugins/Experimental/CommonUI/Source/CommonUI/Private/Input/UIActionRouterTypes.cpp b/Engine/Plugins/Experimental/CommonUI/Source/CommonUI/Private/Input/UIActionRouterTypes.cpp
--- a/Engine/Plugins/Experimental/CommonUI/Source/CommonUI/Private/Input/UIActionRouterTypes.cpp
+++ b/Engine/Plugins/Experimental/CommonUI/Source/CommonUI/Private/Input/UIActionRouterTypes.cpp
@@ -413,6 +413,11 @@ FName GetActionNameFromLegacyRow(const FDataTableRowHandle& LegacyRow)
 
 FName FBindUIActionArgs::GetActionName() const
 {
+	if (InputAction.IsValid())
+	{
+		return InputAction->GetFName();
+	}
+	
 	return LegacyActionTableRow.IsNull() ? ActionTag.GetTagName() : GetActionNameFromLegacyRow(LegacyActionTableRow);
 }
 
