Move CommonUI out of Experimental and into Beta

diff --git a/Engine/Plugins/Runtime/CommonUI/Config/BaseCommonUI.ini b/Engine/Plugins/Runtime/CommonUI/Config/BaseCommonUI.ini
new file mode 100644
--- /dev/null
+++ b/Engine/Plugins/Runtime/CommonUI/Config/BaseCommonUI.ini
@@ -0,0 +1,4 @@
+[CoreRedirects]
++ClassRedirects=(OldName="CommonUIInputData",NewName="/Script/CommonInput.CommonUIInputData")
++ClassRedirects=(OldName="CommonUIKeyBrushDisplayData",NewName="/Script/CommonInput.CommonInputBaseControllerData")
++ClassRedirects=(OldName="CommonInputControllerData",NewName="/Script/CommonInput.CommonInputBaseControllerData")
\ No newline at end of file
diff --git a/Engine/Plugins/Runtime/CommonUI/Source/CommonInput/Public/CommonInputSubsystem.h b/Engine/Plugins/Runtime/CommonUI/Source/CommonInput/Public/CommonInputSubsystem.h
new file mode 100644
--- /dev/null
+++ b/Engine/Plugins/Runtime/CommonUI/Source/CommonInput/Public/CommonInputSubsystem.h
@@ -0,0 +1,182 @@
+// Copyright Epic Games, Inc. All Rights Reserved.
+
+#pragma once
+
+#include "Subsystems/LocalPlayerSubsystem.h"
+#include "CommonInputBaseTypes.h"
+#include "Containers/Ticker.h"
+
+#if UE_ENABLE_INCLUDE_ORDER_DEPRECATED_IN_5_2
+#include "Framework/Application/SlateApplication.h"
+#include "Styling/SlateTypes.h"
+#endif
+
+#include "CommonInputSubsystem.generated.h"
+
+class UWidget;
+class ULocalPlayer;
+class APlayerController;
+class FCommonInputPreprocessor;
+class FSlateUser;
+class UCommonInputActionDomainTable;
+
+DECLARE_DYNAMIC_MULTICAST_DELEGATE_OneParam(FInputMethodChangedDelegate, ECommonInputType, bNewInputType);
+DECLARE_MULTICAST_DELEGATE_ThreeParams(FPlatformInputSupportOverrideDelegate, ULocalPlayer*, ECommonInputType, bool&);
+DECLARE_EVENT_OneParam(FCommonInputPreprocessor, FGamepadChangeDetectedEvent, FName);
+
+UCLASS(DisplayName = "CommonInput")
+class COMMONINPUT_API UCommonInputSubsystem : public ULocalPlayerSubsystem
+{
+	GENERATED_BODY()
+
+public:
+	static UCommonInputSubsystem* Get(const ULocalPlayer* LocalPlayer);
+
+	UCommonInputSubsystem();
+	
+	virtual void Initialize(FSubsystemCollectionBase& Collection) override;
+	virtual void Deinitialize() override;
+
+	DECLARE_EVENT_OneParam(UCommonInputSubsystem, FInputMethodChangedEvent, ECommonInputType);
+	FInputMethodChangedEvent OnInputMethodChangedNative;
+
+	FGamepadChangeDetectedEvent& GetOnGamepadChangeDetected();
+
+	void SetInputTypeFilter(ECommonInputType InputType, FName Reason, bool Filter);
+	bool GetInputTypeFilter(ECommonInputType InputType) const;
+
+	/**  */
+	void AddOrRemoveInputTypeLock(FName InReason, ECommonInputType InInputType, bool bAddLock);
+
+	UFUNCTION(BlueprintCallable, Category = CommonInputSubsystem)
+	bool IsInputMethodActive(ECommonInputType InputMethod) const;
+
+	/** The current input type based on the last input received on the device. */
+	UFUNCTION(BlueprintCallable, Category = CommonInputSubsystem)
+	ECommonInputType GetCurrentInputType() const;
+
+	/** The default input type for the current platform. */
+	UFUNCTION(BlueprintCallable, Category = CommonInputSubsystem)
+	ECommonInputType GetDefaultInputType() const;
+
+	UFUNCTION(BlueprintCallable, Category = CommonInputSubsystem)
+	void SetCurrentInputType(ECommonInputType NewInputType);
+
+	UFUNCTION(BlueprintCallable, Category = CommonInputSubsystem)
+	const FName GetCurrentGamepadName() const;
+
+	UFUNCTION(BlueprintCallable, Category = CommonInputSubsystem)
+	void SetGamepadInputType(const FName InGamepadInputType);
+
+	UFUNCTION(BlueprintCallable, Category = CommonInputSubsystem)
+	bool IsUsingPointerInput() const;
+
+	/** Should display indicators for the current input device on screen.  This is needed when capturing videos, but we don't want to reveal the capture source device. */
+	UFUNCTION(BlueprintCallable, Category = CommonInputSubsystem)
+	bool ShouldShowInputKeys() const;
+
+	void SetActionDomainTable(TObjectPtr<UCommonInputActionDomainTable> Table) { ActionDomainTable = Table; }
+
+	TObjectPtr<UCommonInputActionDomainTable> GetActionDomainTable() const { return ActionDomainTable; }
+
+	/** Returns true if the specified key can be present on both a mobile device and mobile gamepads */
+	static bool IsMobileGamepadKey(const FKey& InKey);
+
+	/** Returns true if the current platform supports a hardware cursor */
+	bool PlatformSupportsHardwareCursor() const;
+
+	void SetCursorPosition(FVector2D NewPosition, bool bForce);
+
+	void UpdateCursorPosition(TSharedRef<FSlateUser> SlateUser, const FVector2D& NewPosition, bool bForce = false);
+
+	/** Getter */
+	bool GetIsGamepadSimulatedClick() const;
+
+	/** Setter */
+	void SetIsGamepadSimulatedClick(bool bNewIsGamepadSimulatedClick);
+
+	/** 
+	* Gets the delegate that allows external systems to override which input methods are supported on this current platform.
+	* @param LocalPlayer								The Local Player.
+	* @param InputType									The current input type that is being tested.
+	* @param InOutCurrentPlatformInputSupportState		The state of if we support the input type as set by PlatformSupportsInputType() or the previous callee of this delegate.
+	* 
+	* Note : Calling order is not guaranteed. Also, keep in mind that you might need to honor the previous callee's request to not support the input type being tested.
+	*/
+	static FPlatformInputSupportOverrideDelegate& GetOnPlatformInputSupportOverride() { return OnPlatformInputSupportOverride; }
+
+protected:
+	virtual TSharedPtr<FCommonInputPreprocessor> MakeInputProcessor();
+
+	ECommonInputType LockInput(ECommonInputType InputToLock) const;
+
+	UFUNCTION()
+	void BroadcastInputMethodChanged();
+
+protected:
+	TSharedPtr<FCommonInputPreprocessor> CommonInputPreprocessor;
+
+private:
+	bool Tick(float DeltaTime);
+
+	void ShouldShowInputKeysChanged(IConsoleVariable* Var);
+
+	FVector2D ClampPositionToViewport(const FVector2D& InPosition) const;
+
+	/** Returns true if the current platform supports the input type */
+	bool PlatformSupportsInputType(ECommonInputType InInputType) const;
+
+	bool CheckForInputMethodThrashing(ECommonInputType NewInputType);
+
+	FTSTicker::FDelegateHandle TickHandle;
+
+	UPROPERTY(BlueprintAssignable, Category = CommonInputSubsystem, meta = (AllowPrivateAccess))
+	FInputMethodChangedDelegate OnInputMethodChanged;
+
+	UPROPERTY(Transient)
+	int32 NumberOfInputMethodChangesRecently = 0;
+
+	UPROPERTY(Transient)
+	double LastInputMethodChangeTime = 0;
+
+	UPROPERTY(Transient)
+	double LastTimeInputMethodThrashingBegan = 0;
+
+	/**  */
+	UPROPERTY(Transient)
+	ECommonInputType LastInputType;
+
+	/**  */
+	UPROPERTY(Transient)
+	ECommonInputType CurrentInputType;
+
+	/**  */
+	UPROPERTY(Transient)
+	FName GamepadInputType;
+
+	/**  */
+	UPROPERTY(Transient)
+	TMap<FName, ECommonInputType> CurrentInputLocks;
+
+	TOptional<ECommonInputType> CurrentInputLock;
+
+	UPROPERTY(Transient)
+	TObjectPtr<UCommonInputActionDomainTable> ActionDomainTable;
+
+	/** Is the current click simulated by the gamepad's face button down/right (platform dependent) */
+	UPROPERTY(Transient)
+	bool bIsGamepadSimulatedClick;
+
+	/**
+	* The delegate that allows external systems to override which input methods are supported on this current platform.
+	* @param LocalPlayer								The Local Player.
+	* @param InputType									The current input type that is being tested.
+	* @param InOutCurrentPlatformInputSupportState		The state of if we support the input type as set by PlatformSupportsInputType() or the previous callee of this delegate.
+	*
+	* Note : Calling order is not guaranteed. Also, keep in mind that you might need to honor the previous callee's request to not support the input type being tested.
+	*/
+	static FPlatformInputSupportOverrideDelegate OnPlatformInputSupportOverride;
+
+	/** Used to broadcast when input method should be considered changed based on external systems, such as on enhanced input mapping context application*/
+	FScriptDelegate BroadcastInputMethodChangedEvent;
+};
\ No newline at end of file
diff --git a/Engine/Plugins/Runtime/CommonUI/Source/CommonUI/CommonUI.Build.cs b/Engine/Plugins/Runtime/CommonUI/Source/CommonUI/CommonUI.Build.cs
new file mode 100644
--- /dev/null
+++ b/Engine/Plugins/Runtime/CommonUI/Source/CommonUI/CommonUI.Build.cs
@@ -0,0 +1,64 @@
+// Copyright Epic Games, Inc. All Rights Reserved.
+
+using UnrealBuildTool;
+
+public class CommonUI : ModuleRules
+{
+	public CommonUI(ReadOnlyTargetRules Target) : base(Target)
+	{
+		PublicDependencyModuleNames.AddRange(
+			new string[]
+			{
+				"Core",
+				"CoreUObject",
+				"ApplicationCore",
+                "Engine",
+				"InputCore",
+				"Slate",
+                "UMG",
+                "WidgetCarousel",
+				"CommonInput",
+				"EnhancedInput",
+				"MediaAssets",
+				"GameplayTags"
+			}
+		);
+
+		PrivateDependencyModuleNames.AddRange(
+		new string[]
+			{
+				"SlateCore",
+                "Analytics",
+				"AnalyticsET",
+                "EngineSettings",
+				"AudioMixer",
+				"DeveloperSettings"
+			}
+		);
+
+		if (Target.Type == TargetType.Editor)
+        {
+            PublicDependencyModuleNames.AddRange(
+                new string[] {
+					"EditorFramework",
+                    "UnrealEd",
+                }
+            );
+        }
+		PrivateDefinitions.Add("UE_COMMONUI_PLATFORM_KBM_REQUIRES_ATTACHED_MOUSE=" + (bPlatformKBMRequiresAttachedMouse ? "1" : "0"));
+		PrivateDefinitions.Add("UE_COMMONUI_PLATFORM_SUPPORTS_TOUCH=" + (bPlatformSupportsTouch ? "1" : "0"));
+	}
+
+	protected virtual bool bPlatformKBMRequiresAttachedMouse { get { return false; } }
+	protected virtual bool bPlatformSupportsTouch
+	{
+		get
+		{
+			// Support touch testing until touch is supported on desktop
+			return (Target.Platform.IsInGroup(UnrealPlatformGroup.Desktop) && Target.Configuration != UnrealTargetConfiguration.Shipping)
+					// Platforms always with touch support
+					|| (Target.Platform == UnrealTargetPlatform.Android)
+					|| (Target.Platform == UnrealTargetPlatform.IOS);
+		}
+	}
+}
diff --git a/Engine/Plugins/Runtime/CommonUI/Source/CommonUI/Private/CommonButtonTypes.h b/Engine/Plugins/Runtime/CommonUI/Source/CommonUI/Private/CommonButtonTypes.h
new file mode 100644
--- /dev/null
+++ b/Engine/Plugins/Runtime/CommonUI/Source/CommonUI/Private/CommonButtonTypes.h
@@ -0,0 +1,131 @@
+// Copyright Epic Games, Inc. All Rights Reserved.
+
+#pragma once
+
+#include "Widgets/Input/SButton.h"
+
+//////////////////////////////////////////////////////////////////////////
+// SCommonButton
+//////////////////////////////////////////////////////////////////////////
+/**
+ * Lets us disable clicking on a button without disabling hit-testing
+ * Needed because NativeOnMouseEnter is not received by disabled widgets,
+ * but that also disables our anchored tooltips.
+ */
+class COMMONUI_API SCommonButton : public SButton
+{
+public:
+	SLATE_BEGIN_ARGS(SCommonButton)
+		: _Content()
+		, _HAlign(HAlign_Fill)
+		, _VAlign(VAlign_Fill)
+		, _ClickMethod(EButtonClickMethod::DownAndUp)
+		, _TouchMethod(EButtonTouchMethod::DownAndUp)
+		, _PressMethod(EButtonPressMethod::DownAndUp)
+		, _IsFocusable(true)
+		, _IsInteractionEnabled(true)
+	{}
+	SLATE_DEFAULT_SLOT(FArguments, Content)
+		SLATE_STYLE_ARGUMENT(FButtonStyle, ButtonStyle)
+		SLATE_ARGUMENT(EHorizontalAlignment, HAlign)
+		SLATE_ARGUMENT(EVerticalAlignment, VAlign)
+		SLATE_EVENT(FOnClicked, OnClicked)
+		SLATE_EVENT(FOnClicked, OnDoubleClicked)
+		SLATE_EVENT(FSimpleDelegate, OnPressed)
+		SLATE_EVENT(FSimpleDelegate, OnReleased)
+		SLATE_ARGUMENT(EButtonClickMethod::Type, ClickMethod)
+		SLATE_ARGUMENT(EButtonTouchMethod::Type, TouchMethod)
+		SLATE_ARGUMENT(EButtonPressMethod::Type, PressMethod)
+		SLATE_ARGUMENT(bool, IsFocusable)
+		SLATE_EVENT(FSimpleDelegate, OnReceivedFocus)
+		SLATE_EVENT(FSimpleDelegate, OnLostFocus)
+
+		/** Is interaction enabled? */
+		SLATE_ARGUMENT(bool, IsButtonEnabled)
+		SLATE_ARGUMENT(bool, IsInteractionEnabled)
+		SLATE_END_ARGS()
+
+	void Construct(const FArguments& InArgs)
+	{
+		OnDoubleClicked = InArgs._OnDoubleClicked;
+
+		SButton::Construct(SButton::FArguments()
+			.ButtonStyle(InArgs._ButtonStyle)
+			.HAlign(InArgs._HAlign)
+			.VAlign(InArgs._VAlign)
+			.ClickMethod(InArgs._ClickMethod)
+			.TouchMethod(InArgs._TouchMethod)
+			.PressMethod(InArgs._PressMethod)
+			.OnClicked(InArgs._OnClicked)
+			.OnPressed(InArgs._OnPressed)
+			.OnReleased(InArgs._OnReleased)
+			.IsFocusable(InArgs._IsFocusable)
+			.Content()
+			[
+				InArgs._Content.Widget
+			]);
+
+		SetCanTick(false);
+		// Set the hover state to indicate that we want to override the default behavior
+		SetHover(false);
+
+		OnReceivedFocus = InArgs._OnReceivedFocus;
+		OnLostFocus = InArgs._OnLostFocus;
+		bIsButtonEnabled = InArgs._IsButtonEnabled;
+		bIsInteractionEnabled = InArgs._IsInteractionEnabled;
+		bHovered = false;
+	}
+
+	virtual FReply OnMouseButtonDown(const FGeometry& MyGeometry, const FPointerEvent& MouseEvent) override;
+
+	virtual FReply OnMouseButtonDoubleClick(const FGeometry& InMyGeometry, const FPointerEvent& InMouseEvent) override;
+
+	virtual FReply OnMouseButtonUp(const FGeometry& MyGeometry, const FPointerEvent& MouseEvent) override;
+
+	virtual void OnMouseEnter(const FGeometry& MyGeometry, const FPointerEvent& MouseEvent) override;
+
+	virtual void OnMouseLeave(const FPointerEvent& MouseEvent) override;
+
+	virtual FReply OnTouchMoved(const FGeometry& MyGeometry, const FPointerEvent& InTouchEvent) override;
+
+	virtual FReply OnKeyDown(const FGeometry& MyGeometry, const FKeyEvent& InKeyEvent) override;
+
+	virtual FReply OnKeyUp(const FGeometry& MyGeometry, const FKeyEvent& InKeyEvent) override;
+
+	void SetIsButtonEnabled(bool bInIsButtonEnabled);
+
+	void SetIsButtonFocusable(bool bInIsButtonFocusable);
+
+	void SetIsInteractionEnabled(bool bInIsInteractionEnabled);
+
+	bool IsInteractable() const;
+
+	/** Overridden to fire delegate for external listener */
+	virtual FReply OnFocusReceived(const FGeometry& MyGeometry, const FFocusEvent& InFocusEvent);
+
+	virtual void OnFocusLost(const FFocusEvent& InFocusEvent) override;
+
+	virtual int32 OnPaint(const FPaintArgs& Args, const FGeometry& AllottedGeometry, const FSlateRect& MyClippingRect, FSlateWindowElementList& OutDrawElements, int32 LayerId, const FWidgetStyle& InWidgetStyle, bool bParentEnabled) const override;
+
+protected:
+	/** Press the button */
+	virtual void Press() override;
+
+private:
+	FOnClicked OnDoubleClicked;
+
+	/** Delegate fired whenever focus is received */
+	FSimpleDelegate OnReceivedFocus;
+
+	/** Delegate fired whenever focus is lost */
+	FSimpleDelegate OnLostFocus;
+
+	/** True if the button is enabled */
+	bool bIsButtonEnabled;
+
+	/** True if clicking is enabled, to allow for things like double click */
+	bool bIsInteractionEnabled;
+
+	/** True if mouse over the widget */
+	bool bHovered;
+};
\ No newline at end of file
diff --git a/Engine/Plugins/Runtime/CommonUI/Source/CommonUI/Private/CommonGameViewportClient.cpp b/Engine/Plugins/Runtime/CommonUI/Source/CommonUI/Private/CommonGameViewportClient.cpp
new file mode 100644
--- /dev/null
+++ b/Engine/Plugins/Runtime/CommonUI/Source/CommonUI/Private/CommonGameViewportClient.cpp
@@ -0,0 +1,228 @@
+// Copyright Epic Games, Inc. All Rights Reserved.
+
+#include "CommonGameViewportClient.h"
+#include "Engine/Console.h"
+#include "Engine/GameInstance.h"
+#include "Engine/LocalPlayer.h"
+#include "InputKeyEventArgs.h"
+
+#if WITH_EDITOR
+#endif // WITH_EDITOR
+
+#include "Input/CommonUIActionRouterBase.h"
+
+#include UE_INLINE_GENERATED_CPP_BY_NAME(CommonGameViewportClient)
+
+#define LOCTEXT_NAMESPACE ""
+
+static const FName NAME_Typing = FName(TEXT("Typing"));
+static const FName NAME_Open = FName(TEXT("Open"));
+
+UCommonGameViewportClient::UCommonGameViewportClient(FVTableHelper& Helper) : Super(Helper)
+{
+}
+
+UCommonGameViewportClient::~UCommonGameViewportClient()
+{
+}
+
+bool UCommonGameViewportClient::InputKey(const FInputKeyEventArgs& InEventArgs)
+{
+	FInputKeyEventArgs EventArgs = InEventArgs;
+
+	if (IsKeyPriorityAboveUI(EventArgs))
+	{
+		return true;
+	}
+
+	// Check override before UI
+	if (OnOverrideInputKey().IsBound())
+	{
+		if (OnOverrideInputKey().Execute(EventArgs))
+		{
+			return true;
+		}
+	}
+
+	// The input is fair game for handling - the UI gets first dibs
+#if !UE_BUILD_SHIPPING
+	if (ViewportConsole && !ViewportConsole->ConsoleState.IsEqual(NAME_Typing) && !ViewportConsole->ConsoleState.IsEqual(NAME_Open))
+#endif
+	{		
+		FReply Result = FReply::Unhandled();
+		if (!OnRerouteInput().ExecuteIfBound(EventArgs.InputDevice, EventArgs.Key, EventArgs.Event, Result))
+		{
+			HandleRerouteInput(EventArgs.InputDevice, EventArgs.Key, EventArgs.Event, Result);
+		}
+
+		if (Result.IsEventHandled())
+		{
+			return true;
+		}
+	}
+
+	return Super::InputKey(EventArgs);
+}
+
+bool UCommonGameViewportClient::InputAxis(FViewport* InViewport, FInputDeviceId InputDevice, FKey Key, float Delta, float DeltaTime, int32 NumSamples, bool bGamepad)
+{
+	FReply RerouteResult = FReply::Unhandled();
+
+	if (!OnRerouteAxis().ExecuteIfBound(InputDevice, Key, Delta, RerouteResult))
+	{
+		HandleRerouteAxis(InputDevice, Key, Delta, RerouteResult);
+	}
+
+	if (RerouteResult.IsEventHandled())
+	{
+		return true;
+	}
+	return Super::InputAxis(InViewport, InputDevice, Key, Delta, DeltaTime, NumSamples, bGamepad);
+}
+
+bool UCommonGameViewportClient::InputTouch(FViewport* InViewport, int32 ControllerId, uint32 Handle, ETouchType::Type Type, const FVector2D& TouchLocation, float Force, FDateTime DeviceTimestamp, uint32 TouchpadIndex)
+{
+#if !UE_BUILD_SHIPPING
+	if (ViewportConsole != NULL && (ViewportConsole->ConsoleState != NAME_Typing) && (ViewportConsole->ConsoleState != NAME_Open))
+#endif
+	{
+		FReply Result = FReply::Unhandled();
+		if (OnRerouteTouch().ExecuteIfBound(ControllerId, Handle, Type, TouchLocation, Result))
+		{
+			HandleRerouteTouch(ControllerId, Handle, Type, TouchLocation, Result);
+		}
+
+		if (Result.IsEventHandled())
+		{
+			return true;
+		}
+	}
+
+	return Super::InputTouch(InViewport, ControllerId, Handle, Type, TouchLocation, Force, DeviceTimestamp, TouchpadIndex);
+}
+
+void UCommonGameViewportClient::HandleRerouteInput(FInputDeviceId DeviceId, FKey Key, EInputEvent EventType, FReply& Reply)
+{
+	FPlatformUserId OwningPlatformUser = IPlatformInputDeviceMapper::Get().GetUserForInputDevice(DeviceId);
+	ULocalPlayer* LocalPlayer = GameInstance->FindLocalPlayerFromPlatformUserId(OwningPlatformUser);
+	Reply = FReply::Unhandled();
+
+	if (LocalPlayer)
+	{
+		UCommonUIActionRouterBase* ActionRouter = LocalPlayer->GetSubsystem<UCommonUIActionRouterBase>();
+		if (ensure(ActionRouter))
+		{
+			ERouteUIInputResult InputResult = ActionRouter->ProcessInput(Key, EventType);
+			if (InputResult == ERouteUIInputResult::BlockGameInput)
+			{
+				// We need to set the reply as handled otherwise the input won't actually be blocked from reaching the viewport.
+				Reply = FReply::Handled();
+				// Notify interested parties that we blocked the input.
+				OnRerouteBlockedInput().ExecuteIfBound(DeviceId, Key, EventType, Reply);
+			}
+			else if (InputResult == ERouteUIInputResult::Handled)
+			{
+				Reply = FReply::Handled();
+			}
+		}
+	}	
+}
+
+void UCommonGameViewportClient::HandleRerouteInput(int32 ControllerId, FKey Key, EInputEvent EventType, FReply& Reply)
+{
+	// Remap the old int32 ControllerId to the new platform user and input device ID
+	FPlatformUserId UserId = FGenericPlatformMisc::GetPlatformUserForUserIndex(ControllerId);
+	FInputDeviceId DeviceID = INPUTDEVICEID_NONE;
+	IPlatformInputDeviceMapper::Get().RemapControllerIdToPlatformUserAndDevice(ControllerId, UserId, DeviceID);
+	return HandleRerouteInput(DeviceID, Key, EventType, Reply);
+}
+
+void UCommonGameViewportClient::HandleRerouteAxis(FInputDeviceId DeviceId, FKey Key, float Delta, FReply& Reply)
+{
+	// Get the ownign platform user for this input device and their local player
+	FPlatformUserId OwningPlatformUser = IPlatformInputDeviceMapper::Get().GetUserForInputDevice(DeviceId);
+	ULocalPlayer* LocalPlayer = GameInstance->FindLocalPlayerFromPlatformUserId(OwningPlatformUser);
+	
+	Reply = FReply::Unhandled();
+
+	if (LocalPlayer)
+	{
+		UCommonUIActionRouterBase* ActionRouter = LocalPlayer->GetSubsystem<UCommonUIActionRouterBase>();
+		if (ensure(ActionRouter))
+		{
+			// We don't actually use axis inputs that reach the game viewport UI land for anything, we just want block them reaching the game when they shouldn't
+			if (!ActionRouter->CanProcessNormalGameInput())
+			{
+				Reply = FReply::Handled();
+			}
+		}
+	}
+}
+
+void UCommonGameViewportClient::HandleRerouteAxis(int32 ControllerId, FKey Key, float Delta, FReply& Reply)
+{
+	// Remap the old int32 ControllerId to the new platform user and input device ID
+	FPlatformUserId UserId = FGenericPlatformMisc::GetPlatformUserForUserIndex(ControllerId);
+	FInputDeviceId DeviceID = INPUTDEVICEID_NONE;
+	IPlatformInputDeviceMapper::Get().RemapControllerIdToPlatformUserAndDevice(ControllerId, UserId, DeviceID);
+	
+	return HandleRerouteAxis(DeviceID, Key, Delta, Reply);
+}
+
+void UCommonGameViewportClient::HandleRerouteTouch(int32 ControllerId, uint32 TouchId, ETouchType::Type TouchType, const FVector2D& TouchLocation, FReply& Reply)
+{
+	ULocalPlayer* LocalPlayer = GameInstance->FindLocalPlayerFromControllerId(ControllerId);
+	Reply = FReply::Unhandled();
+
+	if (LocalPlayer && TouchId < EKeys::NUM_TOUCH_KEYS)
+	{
+		FKey KeyPressed = EKeys::TouchKeys[TouchId];
+		if (KeyPressed.IsValid())
+		{
+			UCommonUIActionRouterBase* ActionRouter = LocalPlayer->GetSubsystem<UCommonUIActionRouterBase>();
+			if (ensure(ActionRouter))
+			{
+				EInputEvent SimilarInputEvent = IE_MAX;
+				switch (TouchType)
+				{
+				case ETouchType::Began:
+					SimilarInputEvent = IE_Pressed;
+					break;
+				case ETouchType::Ended:
+					SimilarInputEvent = IE_Released;
+					break;
+				default:
+					SimilarInputEvent = IE_Repeat;
+					break;
+				}
+
+				if (ActionRouter->ProcessInput(KeyPressed, SimilarInputEvent) != ERouteUIInputResult::Unhandled)
+				{
+					Reply = FReply::Handled();
+				}
+			}
+		}
+	}
+}
+
+bool UCommonGameViewportClient::IsKeyPriorityAboveUI(const FInputKeyEventArgs& EventArgs)
+{
+#if !UE_BUILD_SHIPPING
+	// First priority goes to the viewport console regardless any state or setting
+	if (ViewportConsole && ViewportConsole->InputKey(EventArgs.InputDevice, EventArgs.Key, EventArgs.Event, EventArgs.AmountDepressed, EventArgs.IsGamepad()))
+	{
+		return true;
+	}
+#endif
+
+	// We'll also treat toggling fullscreen as a system-level sort of input that isn't affected by input filtering
+	if (TryToggleFullscreenOnInputKey(EventArgs.Key, EventArgs.Event))
+	{
+		return true;
+	}
+
+	return false;
+}
+
+#undef LOCTEXT_NAMESPACE
+
diff --git a/Engine/Plugins/Runtime/CommonUI/Source/CommonUI/Private/Input/UIActionRouterTypes.h b/Engine/Plugins/Runtime/CommonUI/Source/CommonUI/Private/Input/UIActionRouterTypes.h
new file mode 100644
--- /dev/null
+++ b/Engine/Plugins/Runtime/CommonUI/Source/CommonUI/Private/Input/UIActionRouterTypes.h
@@ -0,0 +1,229 @@
+// Copyright Epic Games, Inc. All Rights Reserved.
+
+#pragma once
+
+#include "Input/UIActionBindingHandle.h"
+#include "UObject/WeakObjectPtr.h"
+#include "UObject/WeakObjectPtrTemplates.h"
+
+enum EInputEvent : int;
+enum class EProcessHoldActionResult;
+struct FKey;
+struct FUIActionBinding;
+template <typename OptionalType> struct TOptional;
+
+// Note: Everything in here should be considered completely private to each other and CommonUIActionRouter.
+//		They were all originally defined directly in CommonUIActionRouter.cpp, but it was annoying having to scroll around so much.
+
+class SWidget;
+class UWidget;
+class UCommonUIActionRouterBase;
+class UCommonActivatableWidget;
+
+class IInputProcessor;
+struct FBindUIActionArgs;
+enum class ECommonInputType : uint8;
+
+class FActivatableTreeNode;
+struct FCommonInputActionDataBase;
+class UCommonInputSubsystem;
+using FActivatableTreeNodePtr = TSharedPtr<FActivatableTreeNode>;
+using FActivatableTreeNodeRef = TSharedRef<FActivatableTreeNode>;
+
+class FActivatableTreeRoot;
+using FActivatableTreeRootPtr = TSharedPtr<FActivatableTreeRoot>;
+using FActivatableTreeRootRef = TSharedRef<FActivatableTreeRoot>;
+
+class FActionRouterBindingCollection;
+
+DECLARE_LOG_CATEGORY_EXTERN(LogUIActionRouter, Log, All);
+
+//////////////////////////////////////////////////////////////////////////
+// FActionRouterBindingCollection
+//////////////////////////////////////////////////////////////////////////
+
+class COMMONUI_API FActionRouterBindingCollection : public TSharedFromThis<FActionRouterBindingCollection>
+{
+public:
+	virtual ~FActionRouterBindingCollection() {}
+
+	virtual EProcessHoldActionResult ProcessHoldInput(ECommonInputMode ActiveInputMode, FKey Key, EInputEvent InputEvent) const;
+	virtual bool ProcessNormalInput(ECommonInputMode ActiveInputMode, FKey Key, EInputEvent InputEvent) const;
+	virtual bool IsReceivingInput() const { return true; }
+
+	void AddBinding(FUIActionBinding& Binding);
+	
+	void RemoveBindings(const TArray<FUIActionBindingHandle>& WidgetBindings);
+	void RemoveBinding(FUIActionBindingHandle ActionHandle);
+
+	bool HasHoldBindings() const { return HoldBindingsCount > 0; }
+
+	const TArray<FUIActionBindingHandle>& GetActionBindings() const { return ActionBindings; }
+	
+protected:
+	FActionRouterBindingCollection(UCommonUIActionRouterBase& OwningRouter);
+	virtual bool IsWidgetReachableForInput(const UWidget* Widget) const;
+
+	int32 GetOwnerUserIndex() const;
+	int32 GetOwnerControllerId() const;
+	UCommonUIActionRouterBase& GetActionRouter() const { check(ActionRouterPtr.IsValid()); return *ActionRouterPtr; }
+	
+	void DebugDumpActionBindings(FString& OutputStr, int32 IndentSpaces) const;
+
+	/** The set of action bindings contained within this collection */
+	TArray<FUIActionBindingHandle> ActionBindings;
+	
+	/**
+	 * Treat this as guaranteed to be valid and access via GetActionRouter()
+	 * Only kept as a WeakObjectPtr so we can reliably assert in the case it somehow becomes invalid.
+	 */
+	TWeakObjectPtr<UCommonUIActionRouterBase> ActionRouterPtr;
+
+private:
+	//Slate application sends repeat actions only for the last pressed key, so we have to keep track of this last held binding and clear it when we get a new key to hold
+	mutable FUIActionBindingHandle CurrentlyHeldBinding;
+
+	int32 HoldBindingsCount = 0;
+};
+
+//////////////////////////////////////////////////////////////////////////
+// FActivatableTreeNode
+//////////////////////////////////////////////////////////////////////////
+
+class COMMONUI_API FActivatableTreeNode : public FActionRouterBindingCollection
+{
+public:
+	virtual ~FActivatableTreeNode();
+	
+	virtual EProcessHoldActionResult ProcessHoldInput(ECommonInputMode ActiveInputMode, FKey Key, EInputEvent InputEvent) const override;
+	virtual bool ProcessNormalInput(ECommonInputMode ActiveInputMode, FKey Key, EInputEvent InputEvent) const override;
+	virtual bool IsReceivingInput() const override { return bCanReceiveInput && IsWidgetActivated(); }
+
+	bool IsWidgetValid() const;
+	bool IsWidgetActivated() const;
+	bool DoesWidgetSupportActivationFocus() const;
+	void AppendAllActiveActions(TArray<FUIActionBindingHandle>& BoundActions) const;
+
+	UCommonActivatableWidget* GetWidget() { return RepresentedWidget.Get(); }
+	const UCommonActivatableWidget* GetWidget() const { return RepresentedWidget.Get(); }
+
+	TArray<FActivatableTreeNodeRef>& GetChildren() { return Children; }
+	const TArray<FActivatableTreeNodeRef>& GetChildren() const { return Children; }
+	
+	FActivatableTreeNodePtr GetParentNode() const { return Parent.Pin(); }
+	
+	FActivatableTreeNodeRef AddChildNode(UCommonActivatableWidget& InActivatableWidget);
+	void CacheFocusRestorationTarget();
+	TSharedPtr<SWidget> GetFocusFallbackTarget() const;
+
+	bool IsExclusiveParentOfWidget(const TSharedPtr<SWidget>& SlateWidget) const;
+
+	int32 GetLastPaintLayer() const;
+	TOptional<FUIInputConfig> FindDesiredInputConfig() const;
+	TOptional<FUIInputConfig> FindDesiredActionDomainInputConfig() const;
+	FActivationMetadata FindActivationMetadata() const;
+	
+	void SetCanReceiveInput(bool bInCanReceiveInput);
+	
+	void AddScrollRecipient(const UWidget& ScrollRecipient);
+	void RemoveScrollRecipient(const UWidget& ScrollRecipient);
+	void AddInputPreprocessor(const TSharedRef<IInputProcessor>& InputPreprocessor, int32 DesiredIndex);
+
+	FSimpleDelegate OnActivated;
+	FSimpleDelegate OnDeactivated;
+
+protected:
+	FActivatableTreeNode(UCommonUIActionRouterBase& OwningRouter, UCommonActivatableWidget& ActivatableWidget);
+	FActivatableTreeNode(UCommonUIActionRouterBase& OwningRouter, UCommonActivatableWidget& ActivatableWidget, const FActivatableTreeNodeRef& InParent);
+
+	virtual bool IsWidgetReachableForInput(const UWidget* Widget) const override;
+	
+	bool CanReceiveInput() const { return bCanReceiveInput; }
+	virtual void Init();	
+	FActivatableTreeRootRef GetRoot() const;
+
+	void AppendValidScrollRecipients(TArray<const UWidget*>& AllScrollRecipients) const;
+	void DebugDumpRecursive(FString& OutputStr, int32 Depth, bool bIncludeActions, bool bIncludeChildren, bool bIncludeInactive) const;
+	
+	bool IsParentOfWidget(const TSharedPtr<SWidget>& SlateWidget) const;
+	
+private:
+	void HandleWidgetActivated();
+	void HandleWidgetDeactivated();
+	void HandleChildSlateReleased(UCommonActivatableWidget* ChildWidget);
+	
+	void RegisterPreprocessors();
+	void UnregisterPreprocessors();
+
+	bool DoesPathSupportActivationFocus() const;
+	
+#if !UE_BUILD_SHIPPING
+	FString DebugWidgetName;
+#endif
+
+	TWeakObjectPtr<UCommonActivatableWidget> RepresentedWidget;
+	TWeakPtr<FActivatableTreeNode> Parent;
+	TArray<FActivatableTreeNodeRef> Children;
+	TWeakPtr<SWidget> FocusRestorationTarget;
+
+	bool bCanReceiveInput = false;
+
+	struct FPreprocessorRegistration
+	{
+		int32 DesiredIndex = INDEX_NONE;
+		TSharedRef<IInputProcessor> Preprocessor;
+		FPreprocessorRegistration(int32 InDesiredIndex, const TSharedRef<IInputProcessor>& InPreprocessor)
+			: DesiredIndex(InDesiredIndex), Preprocessor(InPreprocessor) 
+		{}
+	};
+	TArray<FPreprocessorRegistration> RegisteredPreprocessors;
+
+	// Mutable so we can keep it clean during normal use
+	mutable TArray<TWeakObjectPtr<const UWidget>> ScrollRecipients;
+};
+
+//////////////////////////////////////////////////////////////////////////
+// FActivatableTreeRoot
+//////////////////////////////////////////////////////////////////////////
+
+class COMMONUI_API FActivatableTreeRoot : public FActivatableTreeNode
+{
+public:
+	static FActivatableTreeRootRef Create(UCommonUIActionRouterBase& OwningRouter, UCommonActivatableWidget& ActivatableWidget);
+	
+	void UpdateLeafNode();
+
+	TArray<const UWidget*> GatherScrollRecipients() const;
+
+	bool UpdateLeafmostActiveNode(FActivatableTreeNodePtr BaseCandidateNode, bool bInApplyConfig = true);
+
+	void DebugDump(FString& OutputStr, bool bIncludeActions, bool bIncludeChildren, bool bIncludeInactive) const;
+
+	FSimpleDelegate OnLeafmostActiveNodeChanged;
+
+	void FocusLeafmostNode();
+
+	void RefreshCachedRestorationTarget();
+
+	void ApplyLeafmostNodeConfig();
+
+protected:
+	virtual void Init() override;
+
+private:
+	FActivatableTreeRoot(UCommonUIActionRouterBase& OwningRouter, UCommonActivatableWidget& ActivatableWidget)
+		: FActivatableTreeNode(OwningRouter, ActivatableWidget)
+	{}
+
+	void HandleInputMethodChanged(ECommonInputType InputMethod);
+
+	void HandleRequestRefreshLeafmostFocus();
+
+	// WeakPtr because the root itself can be the primary active node - results in a circular ref leak using a full SharedPtr here
+	TWeakPtr<FActivatableTreeNode> LeafmostActiveNode;
+};
+
+#if UE_ENABLE_INCLUDE_ORDER_DEPRECATED_IN_5_2
+#include "Input/CommonUIInputSettings.h"
+#include "Input/UIActionBinding.h"
+#endif
diff --git a/Engine/Plugins/Runtime/CommonUI/Source/CommonUI/Public/AnalogSlider.h b/Engine/Plugins/Runtime/CommonUI/Source/CommonUI/Public/AnalogSlider.h
new file mode 100644
--- /dev/null
+++ b/Engine/Plugins/Runtime/CommonUI/Source/CommonUI/Public/AnalogSlider.h
@@ -0,0 +1,48 @@
+// Copyright Epic Games, Inc. All Rights Reserved.
+
+#pragma once
+
+#include "Components/Slider.h"
+
+
+#include "AnalogSlider.generated.h"
+
+enum class ECommonInputType : uint8;
+
+class SAnalogSlider;
+
+/**
+ * A simple widget that shows a sliding bar with a handle that allows you to control the value in a user define range (between 0..1 by default).
+ *
+ * * No Children
+ */
+UCLASS()
+class COMMONUI_API UAnalogSlider : public USlider
+{
+	GENERATED_UCLASS_BODY()
+
+public:
+	/** Called when the value is changed by slider or typing. */
+	UPROPERTY(BlueprintAssignable, Category = "Widget Event")
+	FOnFloatValueChangedEvent OnAnalogCapture;
+
+	// UWidget interface
+	virtual void SynchronizeProperties() override;
+	// End of UWidget interface
+	
+	virtual TSharedRef<SWidget> RebuildWidget() override;
+	virtual void ReleaseSlateResources(bool bReleaseChildren) override;
+	
+	void HandleOnAnalogCapture(float InValue);
+
+	void HandleInputMethodChanged(ECommonInputType CurrentInputType);
+
+protected:
+	TSharedPtr<SAnalogSlider> MyAnalogSlider;
+};
+
+#if UE_ENABLE_INCLUDE_ORDER_DEPRECATED_IN_5_2
+#include "CommonInputSubsystem.h"
+#include "SAnalogSlider.h"
+#include "Types/NavigationMetaData.h"
+#endif
diff --git a/Engine/Plugins/Runtime/CommonUI/Source/CommonUI/Public/CommonActionWidget.h b/Engine/Plugins/Runtime/CommonUI/Source/CommonUI/Public/CommonActionWidget.h
new file mode 100644
--- /dev/null
+++ b/Engine/Plugins/Runtime/CommonUI/Source/CommonUI/Public/CommonActionWidget.h
@@ -0,0 +1,142 @@
+// Copyright Epic Games, Inc. All Rights Reserved.
+
+#pragma once
+
+#include "Components/Widget.h"
+#include "Input/UIActionBindingHandle.h"
+
+#include "Engine/DataTable.h"
+#include "CommonActionWidget.generated.h"
+
+class UCommonInputSubsystem;
+enum class ECommonInputType : uint8;
+struct FCommonInputActionDataBase;
+
+class SBox;
+class SImage;
+class UMaterialInstanceDynamic;
+
+/**
+ * A widget that shows a platform-specific icon for the given input action.
+ */
+UCLASS(BlueprintType, Blueprintable)
+class COMMONUI_API UCommonActionWidget: public UWidget
+{
+	GENERATED_UCLASS_BODY()
+
+public:
+
+	//UObject interface
+	virtual void Serialize(FArchive& Ar) override;
+	//~ End UObject Interface
+	
+	/** Begin UWidget */
+	virtual TSharedRef<SWidget> RebuildWidget() override;
+	virtual void ReleaseSlateResources(bool bReleaseChildren) override;
+	virtual void SynchronizeProperties() override;
+	/** End UWidet */
+	
+	UFUNCTION(BlueprintCallable, Category = CommonActionWidget)
+	virtual FSlateBrush GetIcon() const;
+
+	UFUNCTION(BlueprintCallable, Category = CommonActionWidget)
+	FText GetDisplayText() const;
+
+	UFUNCTION(BlueprintCallable, Category = CommonActionWidget)
+	void SetEnhancedInputAction(UInputAction* InInputAction);
+
+	UFUNCTION(BlueprintCallable, Category = CommonActionWidget)
+	void SetInputAction(FDataTableRowHandle InputActionRow);
+
+	void SetInputActionBinding(FUIActionBindingHandle BindingHandle);
+
+	UFUNCTION(BlueprintCallable, Category = CommonActionWidget)
+	void SetInputActions(TArray<FDataTableRowHandle> NewInputActions);
+
+	UFUNCTION(BlueprintCallable, Category = CommonActionWidget)
+	void SetIconRimBrush(FSlateBrush InIconRimBrush);
+
+	UFUNCTION(BlueprintCallable, Category = CommonActionWidget)
+	bool IsHeldAction() const;
+
+	DECLARE_DYNAMIC_MULTICAST_DELEGATE_OneParam(FOnInputMethodChanged, bool, bUsingGamepad);
+	UPROPERTY(BlueprintAssignable, Category = CommonActionWidget)
+	FOnInputMethodChanged OnInputMethodChanged;
+
+	/**
+	 * The material to use when showing held progress, the progress will be sent using the material parameter
+	 * defined by ProgressMaterialParam and the value will range from 0..1.
+	 **/
+	UPROPERTY(EditAnywhere, BlueprintReadWrite, Category = CommonActionWidget)
+	FSlateBrush ProgressMaterialBrush;
+
+	/** The material parameter on ProgressMaterialBrush to update the held percentage.  This value will be 0..1. */
+	UPROPERTY(EditAnywhere, BlueprintReadOnly, Category = CommonActionWidget)
+	FName ProgressMaterialParam;
+
+	UPROPERTY(EditAnywhere, BlueprintReadOnly, Category = CommonActionWidget)
+	FSlateBrush IconRimBrush;
+
+#if WITH_EDITOR
+	virtual const FText GetPaletteCategory() override;
+#endif
+
+	void OnActionProgress(float HeldPercent);
+	void OnActionComplete();
+	void SetProgressMaterial(const FSlateBrush& InProgressMaterialBrush, const FName& InProgressMaterialParam);
+	void SetHidden(bool bAlwaysHidden);
+
+protected:
+	/**
+	 * List all the input actions that this common action widget is intended to represent.  In some cases you might have multiple actions
+	 * that you need to represent as a single entry in the UI.  For example - zoom, might be mouse wheel up or down, but you just need to
+	 * show a single icon for Up & Down on the mouse, this solves that problem.
+	 */
+	UPROPERTY(EditAnywhere, BlueprintReadOnly, Category = CommonActionWidget, meta = (RowType = "/Script/CommonUI.CommonInputActionDataBase", TitleProperty = "RowName"))
+	TArray<FDataTableRowHandle> InputActions;
+
+	/**
+	 * Input Action this common action widget is intended to represent. Optional if using EnhancedInputs
+	 */
+	UPROPERTY(EditAnywhere, BlueprintReadOnly, Category = CommonActionWidget, meta = (EditCondition = "CommonInput.CommonInputSettings.IsEnhancedInputSupportEnabled", EditConditionHides))
+	TObjectPtr<class UInputAction> EnhancedInputAction;
+
+	FUIActionBindingHandle DisplayedBindingHandle;
+
+#if WITH_EDITORONLY_DATA
+	UPROPERTY()
+	FDataTableRowHandle InputActionDataRow_DEPRECATED;
+#endif
+
+protected:
+	UCommonInputSubsystem* GetInputSubsystem() const;
+	const FCommonInputActionDataBase* GetInputActionData() const;
+
+	void UpdateActionWidget();
+
+	virtual void OnWidgetRebuilt() override;
+	
+	void UpdateBindingHandleInternal(FUIActionBindingHandle BindingHandle);
+
+	void ListenToInputMethodChanged(bool bListen = true);
+	void HandleInputMethodChanged(ECommonInputType InInputType);
+
+	UPROPERTY(Transient)
+	TObjectPtr<UMaterialInstanceDynamic> ProgressDynamicMaterial;
+	
+	TSharedPtr<SBox> MyKeyBox;
+
+	TSharedPtr<SImage> MyIcon;
+
+	TSharedPtr<SImage> MyProgressImage;
+
+	TSharedPtr<SImage> MyIconRim;
+
+	FSlateBrush Icon;
+
+	bool bAlwaysHideOverride = false;
+};
+
+#if UE_ENABLE_INCLUDE_ORDER_DEPRECATED_IN_5_2
+#include "CommonUITypes.h"
+#endif
diff --git a/Engine/Plugins/Runtime/CommonUI/Source/CommonUI/Public/CommonActivatableWidgetSwitcher.h b/Engine/Plugins/Runtime/CommonUI/Source/CommonUI/Public/CommonActivatableWidgetSwitcher.h
new file mode 100644
--- /dev/null
+++ b/Engine/Plugins/Runtime/CommonUI/Source/CommonUI/Public/CommonActivatableWidgetSwitcher.h
@@ -0,0 +1,46 @@
+// Copyright Epic Games, Inc. All Rights Reserved.
+
+#pragma once
+
+#include "CommonAnimatedSwitcher.h"
+
+
+#include "CommonActivatableWidgetSwitcher.generated.h"
+
+class UCommonActivatableWidget;
+
+/**
+ * An animated switcher that knows about CommonActivatableWidgets. It can also hold other Widgets.
+ */
+UCLASS()
+class COMMONUI_API UCommonActivatableWidgetSwitcher : public UCommonAnimatedSwitcher
+{
+	GENERATED_BODY()
+
+public:
+	virtual void ReleaseSlateResources(bool bReleaseChildren) override;
+
+protected:
+	virtual void OnWidgetRebuilt() override;
+
+	virtual void HandleOutgoingWidget() override;
+	virtual void HandleSlateActiveIndexChanged(int32 ActiveIndex) override;
+
+private:
+	void HandleOwningWidgetActivationChanged(const bool bIsActivated);
+
+	void AttemptToActivateActiveWidget();
+	void DeactivateActiveWidget();
+
+	void BindOwningActivatableWidget();
+	void UnbindOwningActivatableWidget();
+
+	UCommonActivatableWidget* GetOwningActivatableWidget() const;
+
+	TOptional<TWeakObjectPtr<UCommonActivatableWidget>> WeakOwningActivatableWidget;
+};
+
+#if UE_ENABLE_INCLUDE_ORDER_DEPRECATED_IN_5_2
+#include "Misc/Optional.h"
+#include "UObject/WeakObjectPtrTemplates.h"
+#endif
diff --git a/Engine/Plugins/Runtime/CommonUI/Source/CommonUI/Public/CommonAnimatedSwitcher.h b/Engine/Plugins/Runtime/CommonUI/Source/CommonUI/Public/CommonAnimatedSwitcher.h
new file mode 100644
--- /dev/null
+++ b/Engine/Plugins/Runtime/CommonUI/Source/CommonUI/Public/CommonAnimatedSwitcher.h
@@ -0,0 +1,99 @@
+// Copyright Epic Games, Inc. All Rights Reserved.
+
+#pragma once
+
+#include "Components/WidgetSwitcher.h"
+#include "Slate/SCommonAnimatedSwitcher.h"
+
+#include "CommonAnimatedSwitcher.generated.h"
+
+class SOverlay;
+class SSpacer;
+
+/**
+ * A widget switcher that activates / deactivates CommonActivatableWidgets, allowing for associated animations to trigger.
+ */
+UCLASS()
+class COMMONUI_API UCommonAnimatedSwitcher : public UWidgetSwitcher
+{
+	GENERATED_BODY()
+
+public:
+	UCommonAnimatedSwitcher(const FObjectInitializer& ObjectInitializer);
+
+	virtual void ReleaseSlateResources(bool bReleaseChildren) override;
+
+#if WITH_EDITOR
+	virtual const FText GetPaletteCategory() override;
+#endif
+
+	virtual void SetActiveWidgetIndex(int32 Index) override;
+	virtual void SetActiveWidget(UWidget* Widget) override;
+
+	UFUNCTION(BlueprintCallable, Category = "Common Widget Switcher")
+    void ActivateNextWidget(bool bCanWrap);
+
+    UFUNCTION(BlueprintCallable, Category = "Common Widget Switcher")
+    void ActivatePreviousWidget(bool bCanWrap);
+
+    UFUNCTION(BlueprintCallable, Category = "Common Widget Switcher")
+    bool HasWidgets() const;
+
+	UFUNCTION(BlueprintCallable, Category = "Common Widget Switcher")
+	void SetDisableTransitionAnimation(bool bDisableAnimation);
+
+	UFUNCTION(BlueprintCallable, Category = "Common Widget Switcher")
+	bool IsCurrentlySwitching() const;
+
+	/** Is the switcher playing a transition animation? */
+	UFUNCTION(BlueprintCallable, Category = "Common Widget Switcher")
+	bool IsTransitionPlaying() const;
+
+protected:
+	virtual void HandleSlateActiveIndexChanged(int32 ActiveIndex);
+
+	virtual void HandleSlateIsTransitioningChanged(bool bIsTransitioning);
+
+	virtual TSharedRef<SWidget> RebuildWidget() override;
+
+	virtual void HandleOutgoingWidget() {};
+
+public:
+	/** Fires when the active widget displayed by the switcher changes */
+	DECLARE_EVENT_TwoParams(UCommonAnimatedSwitcher, FOnActiveIndexChanged, UWidget*, int32)
+	FOnActiveIndexChanged OnActiveWidgetIndexChanged;
+
+	/** Fires when the switcher changes its transition animation state */
+	DECLARE_EVENT_OneParam(UCommonAnimatedSwitcher, FOnTransitioningChanged, bool)
+	FOnTransitioningChanged OnTransitioningChanged;
+
+protected:
+	/** The type of transition to play between widgets */
+	UPROPERTY(EditAnywhere, BlueprintReadOnly, Category = "Transition")
+	ECommonSwitcherTransition TransitionType;
+
+	/** The curve function type to apply to the transition animation */
+	UPROPERTY(EditAnywhere, BlueprintReadOnly, Category = "Transition")
+	ETransitionCurve TransitionCurveType;
+
+	/** The total duration of a single transition between widgets */
+	UPROPERTY(EditAnywhere, BlueprintReadOnly, Category = "Transition")
+	float TransitionDuration;
+
+	TSharedPtr<SOverlay> MyOverlay;
+	TSharedPtr<SSpacer> MyInputGuard;
+	TSharedPtr<SCommonAnimatedSwitcher> MyAnimatedSwitcher;
+
+	/* If set, transition animations will not play */
+	bool bInstantTransition = false;
+
+	bool bSetOnce = false;
+	bool bCurrentlySwitching = false;
+
+private:
+	void SetActiveWidgetIndex_Internal(int32 Index);
+};
+
+#if UE_ENABLE_INCLUDE_ORDER_DEPRECATED_IN_5_2
+#include "Components/Widget.h"
+#endif
diff --git a/Engine/Plugins/Runtime/CommonUI/Source/CommonUI/Public/CommonBorder.h b/Engine/Plugins/Runtime/CommonUI/Source/CommonUI/Public/CommonBorder.h
new file mode 100644
--- /dev/null
+++ b/Engine/Plugins/Runtime/CommonUI/Source/CommonUI/Public/CommonBorder.h
@@ -0,0 +1,90 @@
+// Copyright Epic Games, Inc. All Rights Reserved.
+
+#pragma once
+
+#include "Components/Border.h"
+#include "CommonBorder.generated.h"
+
+struct FDesignerChangedEventArgs;
+
+/* 
+ * ---- All properties must be EditDefaultsOnly, BlueprintReadOnly !!! -----
+ * We return the CDO to blueprints, so we cannot allow any changes (blueprint doesn't support const variables)
+ */
+UCLASS(Abstract, Blueprintable, ClassGroup = UI, meta = (Category = "Common UI"))
+class COMMONUI_API UCommonBorderStyle : public UObject
+{
+	GENERATED_BODY()
+
+public:
+	UCommonBorderStyle();
+	
+	/** The brush for the background of the border */
+	UPROPERTY(EditDefaultsOnly, BlueprintReadOnly, Category = "Properties")
+	FSlateBrush Background;
+
+	UFUNCTION(BlueprintCallable, Category = "Common Border Style|Getters")
+	void GetBackgroundBrush(FSlateBrush& Brush) const;
+};
+
+/**
+ * Uses the border style template defined in CommonUI project settings by default
+ */
+UCLASS(Config = CommonUI, DefaultConfig, ClassGroup = UI, meta = (Category = "Common UI", DisplayName = "Common Border"))
+class COMMONUI_API UCommonBorder : public UBorder
+{
+	GENERATED_UCLASS_BODY()
+public:
+
+	UFUNCTION(BlueprintCallable, Category = "Common Border")
+	void SetStyle(TSubclassOf<UCommonBorderStyle> InStyle);
+
+	/** References the border style to use */
+	UPROPERTY(EditAnywhere, BlueprintReadOnly, Category = "Common Border", meta = (ExposeOnSpawn = true))
+	TSubclassOf<UCommonBorderStyle> Style;
+
+	/** Turning this on will cause the safe zone size to be removed from this borders content padding down to the minimum specified */
+	UPROPERTY(EditAnywhere, BlueprintReadOnly, Category = "Common Border")
+	bool bReducePaddingBySafezone;
+
+	/** The minimum padding we will reduce to when the safezone grows */
+	UPROPERTY(EditAnywhere, BlueprintReadOnly, Category = "Common Border", meta = (EditCondition = "bReducePaddingBySafezone"))
+	FMargin MinimumPadding;
+
+#if WITH_EDITORONLY_DATA
+	/** Used to track widgets that were created before changing the default style pointer to null */
+	UPROPERTY()
+	bool bStyleNoLongerNeedsConversion;
+#endif
+
+protected:
+	virtual void PostLoad() override;
+
+	// UWidget interface
+	virtual TSharedRef<SWidget> RebuildWidget() override;
+	virtual void ReleaseSlateResources(bool bReleaseChildren) override;
+	virtual void SynchronizeProperties() override;
+	// End of UWidget interface
+
+	void SafeAreaUpdated();
+	void DebugSafeAreaUpdated(const FMargin& NewSafeZone, bool bShouldRecacheMetrics) { SafeAreaUpdated(); };
+#if WITH_EDITOR
+	virtual void OnCreationFromPalette() override;
+	const FText GetPaletteCategory() override;
+	virtual void OnDesignerChanged(const FDesignerChangedEventArgs& EventArgs) override;
+#endif // WITH_EDITOR
+
+protected:
+#if WITH_EDITOR
+	/** The editor-only size constraint passed in by UMG Designer*/
+	TOptional<FVector2D> DesignerSize;
+#endif
+
+private:
+	const UCommonBorderStyle* GetStyleCDO() const;
+
+};
+
+#if UE_ENABLE_INCLUDE_ORDER_DEPRECATED_IN_5_2
+#include "CommonUITypes.h"
+#endif
diff --git a/Engine/Plugins/Runtime/CommonUI/Source/CommonUI/Public/CommonCustomNavigation.h b/Engine/Plugins/Runtime/CommonUI/Source/CommonUI/Public/CommonCustomNavigation.h
new file mode 100644
--- /dev/null
+++ b/Engine/Plugins/Runtime/CommonUI/Source/CommonUI/Public/CommonCustomNavigation.h
@@ -0,0 +1,43 @@
+// Copyright Epic Games, Inc. All Rights Reserved.
+
+#pragma once
+
+#include "Components/Border.h"
+#include "CommonCustomNavigation.generated.h"
+
+/**
+ * Exposes a bindable event that can be used to stomp default border navigation with custom behaviors.
+ */
+UCLASS(Config = CommonUI, DefaultConfig, ClassGroup = UI, meta = (Category = "Common UI", DisplayName = "Common Custom Navigation"))
+class COMMONUI_API UCommonCustomNavigation : public UBorder
+{
+	GENERATED_UCLASS_BODY()
+
+public:
+	DECLARE_DYNAMIC_DELEGATE_RetVal_OneParam(bool, FOnCustomNavigationEvent, EUINavigation, NavigationType);
+
+	/** Return true if the Navigation has been handled */
+	UPROPERTY(EditAnywhere, Category = Events, meta = (IsBindableEvent = "True"))
+	FOnCustomNavigationEvent OnNavigationEvent;
+
+public:
+
+#if WITH_EDITOR
+	virtual const FText GetPaletteCategory() override;
+#endif
+
+protected:
+
+	//~ Begin UWidget Interface
+	virtual TSharedRef<SWidget> RebuildWidget() override;
+	//~ End UWidget Interface
+
+protected:
+	
+	bool OnNavigation(EUINavigation NavigationType);
+
+};
+
+#if UE_ENABLE_INCLUDE_ORDER_DEPRECATED_IN_5_2
+#include "CommonUITypes.h"
+#endif
diff --git a/Engine/Plugins/Runtime/CommonUI/Source/CommonUI/Public/CommonDateTimeTextBlock.h b/Engine/Plugins/Runtime/CommonUI/Source/CommonUI/Public/CommonDateTimeTextBlock.h
new file mode 100644
--- /dev/null
+++ b/Engine/Plugins/Runtime/CommonUI/Source/CommonUI/Public/CommonDateTimeTextBlock.h
@@ -0,0 +1,66 @@
+// Copyright Epic Games, Inc. All Rights Reserved.
+
+#pragma once
+
+#include "CommonTextBlock.h"
+#include "CommonDateTimeTextBlock.generated.h"
+
+UCLASS(BlueprintType)
+class COMMONUI_API UCommonDateTimeTextBlock : public UCommonTextBlock
+{
+	GENERATED_BODY()
+
+public:
+	UCommonDateTimeTextBlock(const FObjectInitializer& ObjectInitializer);
+
+	DECLARE_EVENT(UCommonDateTimeTextBlock, FOnTimeCountDownCompletion);
+	FOnTimeCountDownCompletion& OnTimeCountDownCompletion() const { return OnTimeCountDownCompletionEvent; }
+
+#if WITH_EDITOR
+	const FText GetPaletteCategory() override;
+#endif // WITH_EDITOR
+
+	virtual void SynchronizeProperties() override;
+
+	UFUNCTION(BlueprintCallable, Category = "DateTime Text Block")
+	void SetDateTimeValue(const FDateTime InDateTime, bool bShowAsCountdown, float InRefreshDelay = 1.0f);
+
+	UFUNCTION(BlueprintCallable, Category = "DateTime Text Block")
+	void SetTimespanValue(const FTimespan InTimespan);
+
+	UFUNCTION(BlueprintCallable, Category = "DateTime Text Block")
+	void SetCountDownCompletionText(const FText InCompletionText);
+
+	UFUNCTION(BlueprintCallable, Category = "DateTime Text Block")
+	FDateTime GetDateTime() const;
+
+protected:
+	
+	void UpdateUnderlyingText();
+
+	virtual bool ShouldClearTimer(const FTimespan& TimeRemaining) const;
+
+	virtual TOptional<FText> FormatTimespan(const FTimespan& InTimespan) const;
+	virtual TOptional<FText> FormatDateTime(const FDateTime& InDateTime) const;
+
+	int32 GetLastDaysCount() const { return LastDaysCount; }
+	int32 GetLastHoursCount() const { return LastHoursCount; }
+
+private:
+
+	// Timer handle for timer-based ticking based on InterpolationUpdateInterval.
+	FTimerHandle TimerTickHandle;
+	float LastTimerTickTime;
+
+	FDateTime DateTime;
+	bool bShowAsCountdown;
+	int32 LastDaysCount;
+	int32 LastHoursCount;
+
+	bool bUseCountdownCompletionText;
+	FText CountdownCompletionText;
+
+	mutable FOnTimeCountDownCompletion OnTimeCountDownCompletionEvent;
+
+	void TimerTick();
+};
\ No newline at end of file
diff --git a/Engine/Plugins/Runtime/CommonUI/Source/CommonUI/Public/CommonGameViewportClient.h b/Engine/Plugins/Runtime/CommonUI/Source/CommonUI/Public/CommonGameViewportClient.h
new file mode 100644
--- /dev/null
+++ b/Engine/Plugins/Runtime/CommonUI/Source/CommonUI/Public/CommonGameViewportClient.h
@@ -0,0 +1,74 @@
+// Copyright Epic Games, Inc. All Rights Reserved.
+
+#pragma once
+
+#include "Engine/GameViewportClient.h"
+#include "CommonGameViewportClient.generated.h"
+
+class FReply;
+
+DECLARE_DELEGATE_FourParams(FOnRerouteInputDelegate, FInputDeviceId /* InputDeviceId */, FKey /* Key */, EInputEvent /* EventType */, FReply& /* Reply */);
+DECLARE_DELEGATE_FourParams(FOnRerouteAxisDelegate, FInputDeviceId /* InputDeviceId */, FKey /* Key */, float /* Delta */, FReply& /* Reply */);
+DECLARE_DELEGATE_FiveParams(FOnRerouteTouchDelegate, int32 /* ControllerId */, uint32 /* TouchId */, ETouchType::Type /* TouchType */, const FVector2D& /* TouchLocation */, FReply& /* Reply */);
+
+/**
+* CommonUI Viewport to reroute input to UI first. Needed to allow CommonUI to route / handle inputs.
+*/
+UCLASS(Within = Engine, transient, config = Engine)
+class COMMONUI_API UCommonGameViewportClient : public UGameViewportClient
+{
+	GENERATED_BODY()
+
+public:
+
+	UCommonGameViewportClient(FVTableHelper& Helper);
+	virtual ~UCommonGameViewportClient();
+
+	// UGameViewportClient interface begin
+	virtual bool InputKey(const FInputKeyEventArgs& EventArgs) override;
+	virtual bool InputAxis(FViewport* InViewport, FInputDeviceId InputDevice, FKey Key, float Delta, float DeltaTime, int32 NumSamples, bool bGamepad) override;
+	virtual bool InputTouch(FViewport* InViewport, int32 ControllerId, uint32 Handle, ETouchType::Type Type, const FVector2D& TouchLocation, float Force, FDateTime DeviceTimestamp, uint32 TouchpadIndex) override;
+	// UGameViewportClient interface end
+
+	FOnRerouteInputDelegate& OnRerouteInput() { return RerouteInput; }
+	FOnRerouteAxisDelegate& OnRerouteAxis() { return RerouteAxis; }
+	FOnRerouteTouchDelegate& OnRerouteTouch() { return RerouteTouch; }
+
+	FOnRerouteInputDelegate& OnRerouteBlockedInput() { return RerouteBlockedInput; }
+
+	/** Default Handler for Key input. */
+	UE_DEPRECATED(5.1, "This version of HandleRerouteInput has been deprecated. Please use the version that takes an FInputDeviceId instead")
+	virtual void HandleRerouteInput(int32 ControllerId, FKey Key, EInputEvent EventType, FReply& Reply);
+
+	/** Default Handler for Key input. */
+	virtual void HandleRerouteInput(FInputDeviceId DeviceId, FKey Key, EInputEvent EventType, FReply& Reply);
+
+	/** Default Handler for Axis input. */
+	UE_DEPRECATED(5.1, "This version of HandleRerouteAxis has been deprecated. Please use the version that takes an FInputDeviceId instead")
+	virtual void HandleRerouteAxis(int32 ControllerId, FKey Key, float Delta, FReply& Reply);
+
+	/** Default Handler for Axis input. */
+	virtual void HandleRerouteAxis(FInputDeviceId DeviceId, FKey Key, float Delta, FReply& Reply);
+
+	/** Default Handler for Touch input. */
+	virtual void HandleRerouteTouch(int32 ControllerId, uint32 TouchId, ETouchType::Type TouchType, const FVector2D& TouchLocation, FReply& Reply);
+
+protected:
+	
+	/** Console window & fullscreen shortcut have higher priority than UI */
+	virtual bool IsKeyPriorityAboveUI(const FInputKeyEventArgs& EventArgs);
+
+	FOnRerouteInputDelegate RerouteInput;
+	FOnRerouteAxisDelegate RerouteAxis;
+	FOnRerouteTouchDelegate RerouteTouch;
+
+	FOnRerouteInputDelegate RerouteBlockedInput;
+};
+
+#if UE_ENABLE_INCLUDE_ORDER_DEPRECATED_IN_5_2
+#include "CoreMinimal.h"
+#include "Input/Reply.h"
+#include "InputCoreTypes.h"
+#include "InputKeyEventArgs.h"
+#include "UObject/ObjectMacros.h"
+#endif
diff --git a/Engine/Plugins/Runtime/CommonUI/Source/CommonUI/Public/CommonHierarchicalScrollBox.h b/Engine/Plugins/Runtime/CommonUI/Source/CommonUI/Public/CommonHierarchicalScrollBox.h
new file mode 100644
--- /dev/null
+++ b/Engine/Plugins/Runtime/CommonUI/Source/CommonUI/Public/CommonHierarchicalScrollBox.h
@@ -0,0 +1,25 @@
+// Copyright Epic Games, Inc. All Rights Reserved.
+
+#pragma once
+
+#include "Components/ScrollBox.h"
+#include "CommonHierarchicalScrollBox.generated.h"
+
+/**
+ * An arbitrary scrollable collection of widgets.  Great for presenting 10-100 widgets in a list.  Doesn't support virtualization.
+ */
+UCLASS()
+class COMMONUI_API UCommonHierarchicalScrollBox : public UScrollBox
+{
+	GENERATED_UCLASS_BODY()
+
+protected:
+	//~ Begin UWidget Interface
+	virtual TSharedRef<SWidget> RebuildWidget() override;
+	//~ End UWidget Interface
+};
+
+#if UE_ENABLE_INCLUDE_ORDER_DEPRECATED_IN_5_2
+#include "CoreMinimal.h"
+#include "UObject/ObjectMacros.h"
+#endif
diff --git a/Engine/Plugins/Runtime/CommonUI/Source/CommonUI/Public/CommonLazyImage.h b/Engine/Plugins/Runtime/CommonUI/Source/CommonUI/Public/CommonLazyImage.h
new file mode 100644
--- /dev/null
+++ b/Engine/Plugins/Runtime/CommonUI/Source/CommonUI/Public/CommonLazyImage.h
@@ -0,0 +1,108 @@
+// Copyright Epic Games, Inc. All Rights Reserved.
+
+#pragma once
+
+#include "CommonLoadGuard.h"
+#include "Components/Image.h"
+
+#include "CommonLazyImage.generated.h"
+
+class UTexture;
+
+class UCommonMcpItemDefinition;
+
+/**
+ * A special Image widget that can show unloaded images and takes care of the loading for you!
+ * 
+ * UCommonLazyImage is another wrapper for SLoadGuard, but it only handles image loading and 
+ * a throbber during loading.
+ * 
+ * If this class changes to show any text, by default it will have CoreStyle styling
+ */
+UCLASS()
+class COMMONUI_API UCommonLazyImage : public UImage
+{
+	GENERATED_UCLASS_BODY()
+
+public:
+	virtual void SetBrush(const FSlateBrush& InBrush) override;
+	virtual void SetBrushFromAsset(USlateBrushAsset* Asset) override;
+	virtual void SetBrushFromTexture(UTexture2D* Texture, bool bMatchSize = false) override;
+	virtual void SetBrushFromTextureDynamic(UTexture2DDynamic* Texture, bool bMatchSize = false) override;
+	virtual void SetBrushFromMaterial(UMaterialInterface* Material) override;
+
+	/** Set the brush from a lazy texture asset pointer - will load the texture as needed. */
+	UFUNCTION(BlueprintCallable, Category = LazyImage)
+	void SetBrushFromLazyTexture(const TSoftObjectPtr<UTexture2D>& LazyTexture, bool bMatchSize = false);
+
+	/** Set the brush from a lazy material asset pointer - will load the material as needed. */
+	UFUNCTION(BlueprintCallable, Category = LazyImage)
+	void SetBrushFromLazyMaterial(const TSoftObjectPtr<UMaterialInterface>& LazyMaterial);
+	
+	/** Set the brush from a string asset ref only - expects the referenced asset to be a texture or material. */
+	UFUNCTION(BlueprintCallable, Category = LazyImage)
+	void SetBrushFromLazyDisplayAsset(const TSoftObjectPtr<UObject>& LazyObject, bool bMatchTextureSize = false);
+
+	UFUNCTION(BlueprintCallable, Category = LazyImage)
+	bool IsLoading() const;
+
+	/**
+	 * Establishes the name of the texture parameter on the currently applied brush material to which textures should be applied.
+	 * Does nothing if the current brush resource object is not a material.
+	 *
+	 * Note: that this is cleared out automatically if/when a new material is established on the brush.
+	 * You must call this function again after doing so if the new material has a texture param.
+	 */
+	UFUNCTION(BlueprintCallable, Category = LazyImage)
+	void SetMaterialTextureParamName(FName TextureParamName);
+
+	FOnLoadGuardStateChangedEvent& OnLoadingStateChanged() { return OnLoadingStateChangedEvent; }
+
+protected:
+	virtual TSharedRef<SWidget> RebuildWidget() override final;
+	virtual void OnWidgetRebuilt() override;
+	virtual void ReleaseSlateResources(bool bReleaseChildren) override;
+	virtual void SynchronizeProperties() override;
+
+	virtual void CancelImageStreaming() override;
+	virtual void OnImageStreamingStarted(TSoftObjectPtr<UObject> SoftObject) override;
+	virtual void OnImageStreamingComplete(TSoftObjectPtr<UObject> LoadedSoftObject) override;
+
+	virtual TSharedRef<SWidget> RebuildImageWidget();
+
+#if WITH_EDITOR
+	virtual const FText GetPaletteCategory() override;
+	virtual bool CanEditChange(const FProperty* InProperty) const override;
+#endif	
+
+	void SetIsLoading(bool bIsLoading);
+
+private:
+	void HandleLoadGuardStateChanged(bool bIsLoading);
+	void ShowDefaultImage();
+
+	void SetBrushObjectInternal(UMaterialInterface* Material);
+	void SetBrushObjectInternal(UTexture* Texture, bool bMatchSize = false);
+
+#if WITH_EDITORONLY_DATA
+	UPROPERTY(EditAnywhere, Category = LoadPreview)
+	bool bShowLoading = false;
+#endif
+
+	UPROPERTY(EditAnywhere, Category = Appearance)
+	FSlateBrush LoadingBackgroundBrush;
+
+	/** 
+	 * If this image uses a material that a texture should be applied to, this is the name of the material param to use.
+	 * I.e. if this property is not blank, the resource object of our brush is a material, and we are given a lazy texture, that texture
+	 * will be assigned to the texture param on the material instead of replacing the material outright on the brush.
+	 */
+	UPROPERTY(EditAnywhere, AdvancedDisplay, Category=Appearance)
+	FName MaterialTextureParamName;
+
+	UPROPERTY(BlueprintAssignable, Category = LazyImage, meta = (DisplayName = "On Loading State Changed", ScriptName = "OnLoadingStateChanged"))
+	FOnLoadGuardStateChangedDynamic BP_OnLoadingStateChanged;
+
+	TSharedPtr<SLoadGuard> MyLoadGuard;
+	FOnLoadGuardStateChangedEvent OnLoadingStateChangedEvent;
+};
diff --git a/Engine/Plugins/Runtime/CommonUI/Source/CommonUI/Public/CommonLazyWidget.h b/Engine/Plugins/Runtime/CommonUI/Source/CommonUI/Public/CommonLazyWidget.h
new file mode 100644
--- /dev/null
+++ b/Engine/Plugins/Runtime/CommonUI/Source/CommonUI/Public/CommonLazyWidget.h
@@ -0,0 +1,80 @@
+// Copyright Epic Games, Inc. All Rights Reserved.
+
+#pragma once
+
+#include "Blueprint/UserWidget.h"
+#include "CommonLoadGuard.h"
+
+#include "Engine/StreamableManager.h"
+#include "CommonLazyWidget.generated.h"
+
+class UCommonMcpItemDefinition;
+
+DECLARE_MULTICAST_DELEGATE_OneParam(FOnLazyContentChangedEvent, UUserWidget*);
+
+/**
+ * A special Image widget that can show unloaded images and takes care of the loading for you!
+ */
+UCLASS()
+class COMMONUI_API UCommonLazyWidget : public UWidget
+{
+	GENERATED_UCLASS_BODY()
+
+public:
+	/**  */
+	UFUNCTION(BlueprintCallable, Category = LazyContent)
+	void SetLazyContent(const TSoftClassPtr<UUserWidget> SoftWidget);
+
+	/**  */
+	UFUNCTION(BlueprintCallable, Category = LazyContent)
+	UUserWidget* GetContent() const { return Content; }
+
+	UFUNCTION(BlueprintCallable, Category = LazyContent)
+	bool IsLoading() const;
+
+	FOnLazyContentChangedEvent& OnContentChanged() { return OnContentChangedEvent; }
+	FOnLoadGuardStateChangedEvent& OnLoadingStateChanged() { return OnLoadingStateChangedEvent; }
+
+protected:
+	virtual TSharedRef<SWidget> RebuildWidget() override;
+	virtual void OnWidgetRebuilt() override;
+	virtual void ReleaseSlateResources(bool bReleaseChildren) override;
+	virtual void SynchronizeProperties() override;
+
+	void SetForceShowSpinner(bool bShowLoading);
+
+	void CancelStreaming();
+	void OnStreamingStarted(TSoftClassPtr<UObject> SoftObject);
+	void OnStreamingComplete(TSoftClassPtr<UObject> LoadedSoftObject);
+
+#if WITH_EDITOR
+	virtual const FText GetPaletteCategory() override;
+#endif	
+
+private:
+	void SetLoadedContent(UUserWidget* InContent);
+	void RequestAsyncLoad(TSoftClassPtr<UObject> SoftObject, TFunction<void()>&& Callback);
+	void RequestAsyncLoad(TSoftClassPtr<UObject> SoftObject, FStreamableDelegate DelegateToCall);
+	void HandleLoadGuardStateChanged(bool bIsLoading);
+
+	UPROPERTY(EditAnywhere, Category = Appearance)
+	FSlateBrush LoadingBackgroundBrush;
+
+	UPROPERTY(Transient)
+	TObjectPtr<UUserWidget> Content;
+
+	TSharedPtr<FStreamableHandle> StreamingHandle;
+	FSoftObjectPath StreamingObjectPath;
+
+	UPROPERTY(BlueprintAssignable, Category = LazyImage, meta = (DisplayName = "On Loading State Changed", ScriptName = "OnLoadingStateChanged"))
+	FOnLoadGuardStateChangedDynamic BP_OnLoadingStateChanged;
+
+	TSharedPtr<SLoadGuard> MyLoadGuard;
+	FOnLoadGuardStateChangedEvent OnLoadingStateChangedEvent;
+
+	FOnLazyContentChangedEvent OnContentChangedEvent;
+};
+
+#if UE_ENABLE_INCLUDE_ORDER_DEPRECATED_IN_5_2
+#include "Components/Widget.h"
+#endif
diff --git a/Engine/Plugins/Runtime/CommonUI/Source/CommonUI/Public/CommonPoolableWidgetInterface.h b/Engine/Plugins/Runtime/CommonUI/Source/CommonUI/Public/CommonPoolableWidgetInterface.h
new file mode 100644
--- /dev/null
+++ b/Engine/Plugins/Runtime/CommonUI/Source/CommonUI/Public/CommonPoolableWidgetInterface.h
@@ -0,0 +1,32 @@
+// Copyright Epic Games, Inc. All Rights Reserved.
+
+#pragma once
+
+#include "UObject/Interface.h"
+#include "CommonPoolableWidgetInterface.generated.h"
+
+UINTERFACE()
+class COMMONUI_API UCommonPoolableWidgetInterface : public UInterface
+{
+	GENERATED_BODY()
+};
+
+/**
+ * Widget pool, if implemented WidgetFactory will attempt to reuse implementing widget objects.
+ */
+class COMMONUI_API ICommonPoolableWidgetInterface
+{
+	GENERATED_BODY()
+
+protected:
+	UFUNCTION(BlueprintNativeEvent, Category = "Common Poolable Widget")
+	void OnAcquireFromPool();
+
+	UFUNCTION(BlueprintNativeEvent, Category = "Common Poolable Widget")
+	void OnReleaseToPool();
+};
+
+#if UE_ENABLE_INCLUDE_ORDER_DEPRECATED_IN_5_2
+#include "CoreMinimal.h"
+#include "UObject/ScriptMacros.h"
+#endif
diff --git a/Engine/Plugins/Runtime/CommonUI/Source/CommonUI/Public/CommonRotator.h b/Engine/Plugins/Runtime/CommonUI/Source/CommonUI/Public/CommonRotator.h
new file mode 100644
--- /dev/null
+++ b/Engine/Plugins/Runtime/CommonUI/Source/CommonUI/Public/CommonRotator.h
@@ -0,0 +1,94 @@
+// Copyright Epic Games, Inc. All Rights Reserved.
+
+#pragma once
+
+#include "CommonButtonBase.h"
+
+
+#include "Input/NavigationReply.h"
+#include "CommonRotator.generated.h"
+
+class UCommonTextBlock;
+
+DECLARE_DYNAMIC_MULTICAST_DELEGATE_OneParam(FOnRotated, int32, Value);
+
+/**
+* A button that can rotate between given text labels.
+*/
+UCLASS(meta = (DisableNativeTick))
+class COMMONUI_API UCommonRotator : public UCommonButtonBase
+{
+	GENERATED_UCLASS_BODY()
+
+public:
+	virtual bool Initialize() override;
+	virtual FNavigationReply NativeOnNavigation(const FGeometry& MyGeometry, const FNavigationEvent& InNavigationEvent, const FNavigationReply& InDefaultReply);
+
+	/** Handle and use controller navigation to rotate text */
+	FNavigationDelegate OnNavigation;
+	TSharedPtr<SWidget> HandleNavigation(EUINavigation UINavigation);
+
+	/** Set the array of texts available */
+	UFUNCTION(BlueprintCallable, Category = "Behavior")
+	void PopulateTextLabels(TArray<FText> Labels);
+
+	/** Gets the current text value of the slider. */
+	UFUNCTION(BlueprintCallable, Category = "Behavior")
+	FText GetSelectedText() const;
+
+	/** Sets the current value of the slider. */
+	UFUNCTION(BlueprintCallable, Category = "Behavior")
+	virtual void SetSelectedItem(int32 InValue);
+
+	/** Gets the current selected index */
+	UFUNCTION(BlueprintCallable, Category = "Behavior")
+	int32 GetSelectedIndex() const { return SelectedIndex; }
+
+	/** Shift the current text left. */
+	UFUNCTION(BlueprintCallable, Category = "Behavior")
+	void ShiftTextLeft();
+
+	/** Shift the current text right. */
+	UFUNCTION(BlueprintCallable, Category = "Behavior")
+	void ShiftTextRight();
+
+public:
+
+	/** Called when the Selected state of this button changes */
+	UPROPERTY(BlueprintAssignable, Category = "Events")
+	FOnRotated OnRotated;
+
+	DECLARE_EVENT_TwoParams(UCommonRotator, FRotationEvent, int32 /*Value*/, bool /*bUserInitiated*/);
+	FRotationEvent OnRotatedEvent;
+
+protected:
+	void ShiftTextLeftInternal(bool bFromNavigation);
+	void ShiftTextRightInternal(bool bFromNavigation);
+
+protected:
+	UFUNCTION(BlueprintImplementableEvent, Category = Events, meta = (DisplayName = "On Options Populated"))
+	void BP_OnOptionsPopulated(int32 Count);
+
+	UFUNCTION(BlueprintImplementableEvent, Category = Events, meta = (DisplayName = "On Options Selected"))
+	void BP_OnOptionSelected(int32 Index);
+
+protected:
+
+	/** The displayed text */
+	UPROPERTY(BlueprintReadOnly, Category = CommonRotator, Meta = (BindWidget))
+	TObjectPtr<UCommonTextBlock> MyText;
+
+	// Holds the array of display texts
+	TArray<FText> TextLabels;
+
+	/** The index of the current text item */
+	int32 SelectedIndex;
+};
+
+#if UE_ENABLE_INCLUDE_ORDER_DEPRECATED_IN_5_2
+#include "CommonBorder.h"
+#include "CommonTextBlock.h"
+#include "CommonUITypes.h"
+#include "CommonUserWidget.h"
+#include "Types/NavigationMetaData.h"
+#endif
diff --git a/Engine/Plugins/Runtime/CommonUI/Source/CommonUI/Public/CommonTileView.h b/Engine/Plugins/Runtime/CommonUI/Source/CommonUI/Public/CommonTileView.h
new file mode 100644
--- /dev/null
+++ b/Engine/Plugins/Runtime/CommonUI/Source/CommonUI/Public/CommonTileView.h
@@ -0,0 +1,24 @@
+// Copyright Epic Games, Inc. All Rights Reserved.
+
+#pragma once
+
+#include "Components/TileView.h"
+#include "CommonTileView.generated.h"
+
+class STableViewBase;
+
+/**
+ * TileView specialized to navigate on focus for consoles & enable scrolling when not focused for touch
+ */
+UCLASS(meta = (DisableNativeTick))
+class COMMONUI_API UCommonTileView : public UTileView
+{
+	GENERATED_BODY()
+
+public:
+	UCommonTileView(const FObjectInitializer& ObjectInitializer);
+
+protected:
+	virtual TSharedRef<STableViewBase> RebuildListWidget() override;
+	virtual UUserWidget& OnGenerateEntryWidgetInternal(UObject* Item, TSubclassOf<UUserWidget> DesiredEntryClass, const TSharedRef<STableViewBase>& OwnerTable) override;
+};
diff --git a/Engine/Plugins/Runtime/CommonUI/Source/CommonUI/Public/CommonUIEditorSettings.h b/Engine/Plugins/Runtime/CommonUI/Source/CommonUI/Public/CommonUIEditorSettings.h
new file mode 100644
--- /dev/null
+++ b/Engine/Plugins/Runtime/CommonUI/Source/CommonUI/Public/CommonUIEditorSettings.h
@@ -0,0 +1,72 @@
+// Copyright Epic Games, Inc. All Rights Reserved.
+
+#pragma once
+
+
+#include "CommonTextBlock.h"
+#include "CommonButtonBase.h"
+#include "CommonBorder.h"
+#include "CommonUIEditorSettings.generated.h"
+
+UCLASS(config = Editor, defaultconfig)
+class COMMONUI_API UCommonUIEditorSettings : public UObject
+{
+	GENERATED_BODY()
+
+public:
+	UCommonUIEditorSettings(const FObjectInitializer& Initializer);
+#if WITH_EDITOR
+	/* Called to load CommonUIEditorSettings data */
+	void LoadData();
+
+	virtual void PostEditChangeProperty(struct FPropertyChangedEvent& PropertyChangedEvent) override;
+	
+	/*  Template Styles - Only accessible in editor builds should be transfered to new widgets in OnCreationFromPalette() overrides */
+	const TSubclassOf<UCommonTextStyle>& GetTemplateTextStyle() const;
+
+	const TSubclassOf<UCommonButtonStyle>& GetTemplateButtonStyle() const;
+
+	const TSubclassOf<UCommonBorderStyle>& GetTemplateBorderStyle() const;
+
+private:
+	void LoadEditorData();
+#endif
+
+private:
+	/** Newly created CommonText Widgets will use this style. */
+	UPROPERTY(config, EditAnywhere, Category = "Text")
+	TSoftClassPtr<UCommonTextStyle> TemplateTextStyle;
+
+	/** Newly created CommonButton Widgets will use this style. */
+	UPROPERTY(config, EditAnywhere, Category = "Buttons")
+	TSoftClassPtr<UCommonButtonStyle> TemplateButtonStyle;
+
+	/** Newly created CommonBorder Widgets will use this style. */
+	UPROPERTY(config, EditAnywhere, Category = "Border")
+	TSoftClassPtr<UCommonBorderStyle> TemplateBorderStyle;
+
+#if WITH_EDITORONLY_DATA
+	UPROPERTY(Transient)
+	TSubclassOf<UCommonTextStyle> TemplateTextStyleClass;
+
+	UPROPERTY(Transient)
+	TSubclassOf<UCommonButtonStyle> TemplateButtonStyleClass;
+
+	UPROPERTY(Transient)
+	TSubclassOf<UCommonBorderStyle> TemplateBorderStyleClass;
+#endif
+private:
+	bool bDefaultDataLoaded;
+};
+
+#if UE_ENABLE_INCLUDE_ORDER_DEPRECATED_IN_5_2
+#include "CommonUITypes.h"
+#include "CoreMinimal.h"
+#include "Engine/StreamableManager.h"
+#include "ICommonUIModule.h"
+#include "InputCoreTypes.h"
+#include "Templates/SharedPointer.h"
+#include "Templates/SubclassOf.h"
+#include "UObject/Object.h"
+#include "UObject/SoftObjectPath.h"
+#endif
diff --git a/Engine/Plugins/Runtime/CommonUI/Source/CommonUI/Public/CommonUILibrary.h b/Engine/Plugins/Runtime/CommonUI/Source/CommonUI/Public/CommonUILibrary.h
new file mode 100644
--- /dev/null
+++ b/Engine/Plugins/Runtime/CommonUI/Source/CommonUI/Public/CommonUILibrary.h
@@ -0,0 +1,29 @@
+// Copyright Epic Games, Inc. All Rights Reserved.
+
+#pragma once
+
+#include "Kismet/BlueprintFunctionLibrary.h"
+#include "CommonUILibrary.generated.h"
+
+class UWidget;
+template <typename T> class TSubclassOf;
+
+UCLASS()
+class COMMONUI_API UCommonUILibrary : public UBlueprintFunctionLibrary
+{
+	GENERATED_UCLASS_BODY()
+
+public:
+
+	/**
+	 * Finds the first parent widget of the given type and returns it, or null if no parent could be found.
+	 */
+	UFUNCTION(BlueprintCallable, Category="Common UI", meta=(DeterminesOutputType=Type))
+	static UWidget* FindParentWidgetOfType(UWidget* StartingWidget, TSubclassOf<UWidget> Type);
+};
+
+#if UE_ENABLE_INCLUDE_ORDER_DEPRECATED_IN_5_2
+#include "Components/Widget.h"
+#include "CoreMinimal.h"
+#include "UObject/ObjectMacros.h"
+#endif
diff --git a/Engine/Plugins/Runtime/CommonUI/Source/CommonUI/Public/CommonUISubsystemBase.h b/Engine/Plugins/Runtime/CommonUI/Source/CommonUI/Public/CommonUISubsystemBase.h
new file mode 100644
--- /dev/null
+++ b/Engine/Plugins/Runtime/CommonUI/Source/CommonUI/Public/CommonUISubsystemBase.h
@@ -0,0 +1,61 @@
+// Copyright Epic Games, Inc. All Rights Reserved.
+
+#pragma once
+
+#include "Subsystems/GameInstanceSubsystem.h"
+
+#include "CommonUISubsystemBase.generated.h"
+
+enum class ECommonInputType : uint8;
+struct FDataTableRowHandle;
+struct FSlateBrush;
+
+class IAnalyticsProviderET;
+class UWidget;
+class ULocalPlayer;
+class UInputAction;
+
+UCLASS(DisplayName = "CommonUI")
+class COMMONUI_API UCommonUISubsystemBase : public UGameInstanceSubsystem
+{
+	GENERATED_BODY()
+
+public:
+	static UCommonUISubsystemBase* Get(const UWidget& Widget);
+
+	virtual bool ShouldCreateSubsystem(UObject* Outer) const override;
+	virtual void Initialize(FSubsystemCollectionBase& Collection) override;
+
+	/** set the analytic provider for the CommonUI Widgets to use */
+	void SetAnalyticProvider(const TSharedPtr<IAnalyticsProviderET>& AnalyticProvider);
+
+	// Gets Action Button Icon for current gamepad
+	UFUNCTION(BlueprintCallable, Category = CommonUISubsystem)
+	FSlateBrush GetInputActionButtonIcon(const FDataTableRowHandle& InputActionRowHandle, ECommonInputType InputType, const FName& GamepadName) const;
+
+	// Gets Action Button Icon for given action and player, enhanced input API currently does not allow input type specification
+	UFUNCTION(BlueprintCallable, Category = CommonUISubsystem)
+	FSlateBrush GetEnhancedInputActionButtonIcon(const UInputAction* InputAction, const ULocalPlayer* LocalPlayer) const;
+
+	/** Analytic Events **/
+
+	//CommonUI.ButtonClicked
+	void FireEvent_ButtonClicked(const FString& InstanceName, const FString& ABTestName, const FString& ExtraData) const;
+
+	//CommonUI.PanelPushed
+	void FireEvent_PanelPushed(const FString& PanelName) const;
+	
+	virtual void SetInputAllowed(bool bEnabled, const FName& Reason, const ULocalPlayer& LocalPlayer);
+	virtual bool IsInputAllowed(const ULocalPlayer* LocalPlayer) const;
+
+private:
+
+	void HandleInputMethodChanged(ECommonInputType bNewInputType);
+
+	TWeakPtr<class IAnalyticsProviderET> AnalyticProviderWeakPtr;
+};
+
+#if UE_ENABLE_INCLUDE_ORDER_DEPRECATED_IN_5_2
+#include "Engine/DataTable.h"
+#include "Styling/SlateBrush.h"
+#endif
diff --git a/Engine/Plugins/Runtime/CommonUI/Source/CommonUI/Public/CommonUserWidget.h b/Engine/Plugins/Runtime/CommonUI/Source/CommonUI/Public/CommonUserWidget.h
new file mode 100644
--- /dev/null
+++ b/Engine/Plugins/Runtime/CommonUI/Source/CommonUI/Public/CommonUserWidget.h
@@ -0,0 +1,94 @@
+// Copyright Epic Games, Inc. All Rights Reserved.
+
+#pragma once
+
+#include "Blueprint/UserWidget.h"
+#include "Input/UIActionBindingHandle.h"
+
+#include "CommonUserWidget.generated.h"
+
+class UCommonInputSubsystem;
+class UCommonUISubsystemBase;
+class FSlateUser;
+
+struct FUIActionTag;
+struct FBindUIActionArgs;
+enum class ECommonInputMode : uint8;
+
+UCLASS(ClassGroup = UI, meta = (Category = "Common UI", DisableNativeTick))
+class COMMONUI_API UCommonUserWidget : public UUserWidget
+{
+	GENERATED_UCLASS_BODY()
+
+public:
+	/** Sets whether or not this widget will consume ALL pointer input that reaches it */
+	UFUNCTION(BlueprintCallable, Category = CommonUserWidget)
+	void SetConsumePointerInput(bool bInConsumePointerInput);
+
+	const TArray<FUIActionBindingHandle>& GetActionBindings() const { return ActionBindings; }
+	const TArray<TWeakObjectPtr<const UWidget>> GetScrollRecipients() const { return ScrollRecipients; }
+
+protected:
+	virtual void OnWidgetRebuilt() override;
+	virtual void NativeDestruct() override;
+	
+	virtual FReply NativeOnMouseButtonDown(const FGeometry& InGeometry, const FPointerEvent& InMouseEvent) override;
+	virtual FReply NativeOnMouseButtonUp(const FGeometry& InGeometry, const FPointerEvent& InMouseEvent) override;
+	virtual FReply NativeOnMouseWheel(const FGeometry& InGeometry, const FPointerEvent& InMouseEvent) override;
+	virtual FReply NativeOnMouseButtonDoubleClick(const FGeometry& InGeometry, const FPointerEvent& InMouseEvent) override;
+	virtual FReply NativeOnTouchGesture(const FGeometry& InGeometry, const FPointerEvent& InGestureEvent) override;
+	virtual FReply NativeOnTouchStarted(const FGeometry& InGeometry, const FPointerEvent& InGestureEvent) override;
+	virtual FReply NativeOnTouchMoved(const FGeometry& InGeometry, const FPointerEvent& InGestureEvent) override;
+	virtual FReply NativeOnTouchEnded(const FGeometry& InGeometry, const FPointerEvent& InGestureEvent) override;
+	
+	UCommonInputSubsystem* GetInputSubsystem() const;
+	UCommonUISubsystemBase* GetUISubsystem() const;
+	TSharedPtr<FSlateUser> GetOwnerSlateUser() const;
+
+	template <typename GameInstanceT = UGameInstance>
+	GameInstanceT& GetGameInstanceChecked() const
+	{
+		GameInstanceT* GameInstance = GetGameInstance<GameInstanceT>();
+		check(GameInstance);
+		return *GameInstance;
+	}
+
+	template <typename PlayerControllerT = APlayerController>
+	PlayerControllerT& GetOwningPlayerChecked() const
+	{
+		PlayerControllerT* PC = GetOwningPlayer<PlayerControllerT>();
+		check(PC);
+		return *PC;
+	}
+
+	void RegisterScrollRecipient(const UWidget& AnalogScrollRecipient);
+	void UnregisterScrollRecipient(const UWidget& AnalogScrollRecipient);
+
+	/** 
+	 * Convenience methods for menu action registrations (any UWidget can register via FCommonUIActionRouter directly, though generally that shouldn't be needed).
+	 * Persistent bindings are *always* listening for input while registered, while normal bindings are only listening when all of this widget's activatable parents are activated.
+	 */
+	FUIActionBindingHandle RegisterUIActionBinding(const FBindUIActionArgs& BindActionArgs);
+
+	void RemoveActionBinding(FUIActionBindingHandle ActionBinding);
+	void AddActionBinding(FUIActionBindingHandle ActionBinding);
+
+	/** True to generally display this widget's actions in the action bar, assuming it has actions. */
+	UPROPERTY(EditAnywhere, BlueprintReadOnly, Category = Input, meta = (AllowPrivateAccess = true))
+	bool bDisplayInActionBar = false;
+
+private:
+
+	/** Set this to true if you don't want any pointer (mouse and touch) input to bubble past this widget */
+	UPROPERTY(EditAnywhere, BlueprintReadOnly, Category = Input, meta = (AllowPrivateAccess = true))
+	bool bConsumePointerInput = false;
+
+private:
+
+	TArray<FUIActionBindingHandle> ActionBindings;
+	TArray<TWeakObjectPtr<const UWidget>> ScrollRecipients;
+};
+
+#if UE_ENABLE_INCLUDE_ORDER_DEPRECATED_IN_5_2
+#include "CommonUITypes.h"
+#endif
diff --git a/Engine/Plugins/Runtime/CommonUI/Source/CommonUI/Public/CommonVideoPlayer.h b/Engine/Plugins/Runtime/CommonUI/Source/CommonUI/Public/CommonVideoPlayer.h
new file mode 100644
--- /dev/null
+++ b/Engine/Plugins/Runtime/CommonUI/Source/CommonUI/Public/CommonVideoPlayer.h
@@ -0,0 +1,93 @@
+// Copyright Epic Games, Inc. All Rights Reserved.
+
+#pragma once
+
+#include "Components/Widget.h"
+
+#include "CommonVideoPlayer.generated.h"
+
+class SImage;
+class UMaterial;
+class UMediaSource;
+class UMediaPlayer;
+class UMediaTexture;
+class UMediaSoundComponent;
+class USoundClass;
+
+enum class EMediaEvent;
+
+UCLASS(ClassGroup = UI, meta = (Category = "Common UI"))
+class COMMONUI_API UCommonVideoPlayer : public UWidget
+{
+	GENERATED_BODY()
+
+public:
+	UCommonVideoPlayer(const FObjectInitializer& Initializer);
+	virtual void PostInitProperties() override;
+
+	void SetVideo(UMediaSource* NewVideo);
+	void Seek(float PlaybackTime);
+	void Close();
+
+	void SetPlaybackRate(float PlaybackRate);
+	void SetLooping(bool bShouldLoopPlayback);
+	void SetIsMuted(bool bInIsMuted);
+
+	void Play();
+	void Reverse();
+	void Pause();
+	void PlayFromStart();
+
+	float GetVideoDuration() const;
+	float GetPlaybackTime() const;
+	float GetPlaybackRate() const;
+
+	bool IsLooping() const;
+	bool IsPaused() const;
+	bool IsPlaying() const;
+	bool IsMuted() const { return bIsMuted; }
+
+	FSimpleMulticastDelegate& OnPlaybackResumed() { return OnPlaybackResumedEvent; }
+	FSimpleMulticastDelegate& OnPlaybackPaused() { return OnPlaybackPausedEvent; }
+	FSimpleMulticastDelegate& OnPlaybackComplete() { return OnPlaybackCompleteEvent; }
+
+protected:
+	virtual TSharedRef<SWidget> RebuildWidget() override;
+	virtual void SynchronizeProperties() override;
+	virtual void ReleaseSlateResources(bool bReleaseChildren) override;
+
+	void PlayInternal() const;
+	const UMediaPlayer& GetMediaPlayer() const { return *MediaPlayer; }
+	virtual void HandleMediaPlayerEvent(EMediaEvent EventType);
+	virtual void PlaybackTick(double InCurrentTime, float InDeltaTime);
+
+private:
+	EActiveTimerReturnType HandlePlaybackTick(double InCurrentTime, float InDeltaTime);
+
+private:
+	UPROPERTY(EditAnywhere, Category = VideoPlayer)
+	TObjectPtr<UMediaSource> Video;
+
+	UPROPERTY(Transient)
+	TObjectPtr<UMediaPlayer> MediaPlayer;
+
+	UPROPERTY(Transient)
+	TObjectPtr<UMediaTexture> MediaTexture;
+
+	UPROPERTY(Transient)
+	TObjectPtr<UMaterial> VideoMaterial;
+
+	UPROPERTY(Transient)
+	TObjectPtr<UMediaSoundComponent> SoundComponent;
+
+	UPROPERTY(Transient)
+	FSlateBrush VideoBrush;
+
+	mutable FSimpleMulticastDelegate OnPlaybackResumedEvent;
+	mutable FSimpleMulticastDelegate OnPlaybackPausedEvent;
+	mutable FSimpleMulticastDelegate OnPlaybackCompleteEvent;
+
+
+	bool bIsMuted = false;
+	TSharedPtr<SImage> MyImage;
+};
diff --git a/Engine/Plugins/Runtime/CommonUI/Source/CommonUI/Public/CommonVisibilitySwitcher.h b/Engine/Plugins/Runtime/CommonUI/Source/CommonUI/Public/CommonVisibilitySwitcher.h
new file mode 100644
--- /dev/null
+++ b/Engine/Plugins/Runtime/CommonUI/Source/CommonUI/Public/CommonVisibilitySwitcher.h
@@ -0,0 +1,97 @@
+// Copyright Epic Games, Inc. All Rights Reserved.
+
+#pragma once
+
+#include "Components/Overlay.h"
+
+#include "CommonVisibilitySwitcher.generated.h"
+
+/**
+ * Basic switcher that toggles visibility on its children to only show one widget at a time. Activates visible widget if possible.
+ */
+UCLASS(meta = (DisableNativeTick))
+class COMMONUI_API UCommonVisibilitySwitcher : public UOverlay
+{
+	GENERATED_BODY()
+
+public:
+
+	UFUNCTION(BlueprintCallable, Category = CommonVisibilitySwitcher)
+	void SetActiveWidgetIndex(int32 Index);
+
+	UFUNCTION(BlueprintCallable, Category = CommonVisibilitySwitcher)
+	int32 GetActiveWidgetIndex() const { return ActiveWidgetIndex; }
+
+	UFUNCTION(BlueprintCallable, Category = CommonVisibilitySwitcher)
+	UWidget* GetActiveWidget() const;
+
+	UFUNCTION(BlueprintCallable, Category = CommonVisibilitySwitcher)
+	void SetActiveWidget(const UWidget* Widget);
+
+	UFUNCTION(BlueprintCallable, Category = CommonVisibilitySwitcher)
+	void IncrementActiveWidgetIndex(bool bAllowWrapping = true);
+
+	UFUNCTION(BlueprintCallable, Category = CommonVisibilitySwitcher)
+	void DecrementActiveWidgetIndex(bool bAllowWrapping = true);
+
+	UFUNCTION(BlueprintCallable, Category = CommonVisibilitySwitcher)
+	void ActivateVisibleSlot();
+
+	UFUNCTION(BlueprintCallable, Category = CommonVisibilitySwitcher)
+	void DeactivateVisibleSlot();
+
+	UWidget* GetWidgetAtIndex(int32 Index) const;
+
+	virtual void SynchronizeProperties() override;
+
+	void MoveChild(int32 CurrentIdx, int32 NewIdx);
+
+	DECLARE_MULTICAST_DELEGATE_OneParam(FOnActiveWidgetIndexChanged, int32)
+	FOnActiveWidgetIndexChanged& OnActiveWidgetIndexChanged() const { return OnActiveWidgetIndexChangedEvent; }
+
+#if WITH_EDITOR
+public:
+
+	virtual const FText GetPaletteCategory() override;
+	virtual void OnDescendantSelectedByDesigner(UWidget* DescendantWidget) override;
+	virtual void OnDescendantDeselectedByDesigner(UWidget* DescendantWidget) override;
+
+	virtual void PostEditChangeProperty(struct FPropertyChangedEvent& PropertyChangedEvent) override;
+
+protected:
+
+	virtual void ValidateCompiledDefaults(class IWidgetCompilerLog& CompileLog) const override;
+
+private:
+
+	int32 DesignTime_ActiveIndex = INDEX_NONE;
+#endif
+
+protected:
+
+	virtual void OnWidgetRebuilt() override;
+	virtual UClass* GetSlotClass() const override;
+	virtual void OnSlotAdded(UPanelSlot* InSlot) override;
+	virtual void OnSlotRemoved(UPanelSlot* InSlot) override;
+
+protected:
+
+	virtual void SetActiveWidgetIndex_Internal(int32 Index, bool bBroadcastChange = true);
+	void ResetSlotVisibilities();
+
+	UPROPERTY(EditAnywhere, Category = CommonVisibilitySwitcher)
+	ESlateVisibility ShownVisibility = ESlateVisibility::SelfHitTestInvisible;
+
+	UPROPERTY(EditAnywhere, Category = CommonVisibilitySwitcher, meta = (ClampMin = -1))
+	int32 ActiveWidgetIndex = 0;
+
+	// Whether or not to automatically activate a slot when it becomes visible
+	UPROPERTY(EditAnywhere, Category = CommonVisibilitySwitcher)
+	bool bAutoActivateSlot = true;
+
+	// Whether or not to activate the first slot if one is added dynamically
+	UPROPERTY(EditAnywhere, Category = CommonVisibilitySwitcher)
+	bool bActivateFirstSlotOnAdding = false;
+
+	mutable FOnActiveWidgetIndexChanged OnActiveWidgetIndexChangedEvent;
+};
diff --git a/Engine/Plugins/Runtime/CommonUI/Source/CommonUI/Public/CommonVisibilitySwitcherSlot.h b/Engine/Plugins/Runtime/CommonUI/Source/CommonUI/Public/CommonVisibilitySwitcherSlot.h
new file mode 100644
--- /dev/null
+++ b/Engine/Plugins/Runtime/CommonUI/Source/CommonUI/Public/CommonVisibilitySwitcherSlot.h
@@ -0,0 +1,34 @@
+// Copyright Epic Games, Inc. All Rights Reserved.
+
+#pragma once
+
+#include "Components/OverlaySlot.h"
+
+#include "CommonVisibilitySwitcherSlot.generated.h"
+
+class SOverlay;
+
+class SBox;
+enum class ESlateVisibility : uint8;
+
+UCLASS()
+class COMMONUI_API UCommonVisibilitySwitcherSlot : public UOverlaySlot
+{
+	GENERATED_BODY()
+
+public:
+
+	UCommonVisibilitySwitcherSlot(const FObjectInitializer& Initializer);
+
+	virtual void BuildSlot(TSharedRef<SOverlay> InOverlay) override;
+
+	virtual void ReleaseSlateResources(bool bReleaseChildren) override;
+
+	void SetSlotVisibility(ESlateVisibility Visibility);
+
+	const TSharedPtr<SBox>& GetVisibilityBox() const { return VisibilityBox; }
+
+private:
+
+	TSharedPtr<SBox> VisibilityBox;
+};
diff --git a/Engine/Plugins/Runtime/CommonUI/Source/CommonUI/Public/CommonVisibilityWidgetBase.h b/Engine/Plugins/Runtime/CommonUI/Source/CommonUI/Public/CommonVisibilityWidgetBase.h
new file mode 100644
--- /dev/null
+++ b/Engine/Plugins/Runtime/CommonUI/Source/CommonUI/Public/CommonVisibilityWidgetBase.h
@@ -0,0 +1,57 @@
+// Copyright Epic Games, Inc. All Rights Reserved.
+
+#pragma once
+
+#include "CommonBorder.h"
+#include "CommonVisibilityWidgetBase.generated.h"
+
+enum class ECommonInputType : uint8;
+enum class ESlateVisibility : uint8;
+
+/**
+ * A container that controls visibility based on Input type and Platform
+ *
+ */
+UCLASS(Deprecated)
+class COMMONUI_API UDEPRECATED_UCommonVisibilityWidgetBase : public UCommonBorder
+{
+	GENERATED_UCLASS_BODY()
+
+public:
+	UPROPERTY(EditAnywhere, EditFixedSize, Category = "Visibility", meta = (GetOptions = GetRegisteredPlatforms))
+	TMap<FName, bool> VisibilityControls;
+
+	UPROPERTY(EditAnywhere, Category = "Visibility")
+	bool bShowForGamepad;
+
+	UPROPERTY(EditAnywhere, Category = "Visibility")
+	bool bShowForMouseAndKeyboard;
+
+	UPROPERTY(EditAnywhere, Category = "Visibility")
+	bool bShowForTouch;
+
+	UPROPERTY(EditAnywhere, Category = "Visibility")
+	ESlateVisibility VisibleType;
+	
+	UPROPERTY(EditAnywhere, Category = "Visibility")
+	ESlateVisibility HiddenType;
+
+protected:
+	// Begin UWidget
+	virtual void OnWidgetRebuilt() override;
+	// End UWidget
+
+	void UpdateVisibility();
+
+	void ListenToInputMethodChanged(bool bListen = true);
+
+	void HandleInputMethodChanged(ECommonInputType input);
+
+	UFUNCTION()
+	static const TArray<FName>& GetRegisteredPlatforms();
+};
+
+#if UE_ENABLE_INCLUDE_ORDER_DEPRECATED_IN_5_2
+#include "CommonUITypes.h"
+#include "CoreMinimal.h"
+#endif
diff --git a/Engine/Plugins/Runtime/CommonUI/Source/CommonUI/Public/CommonVisualAttachment.h b/Engine/Plugins/Runtime/CommonUI/Source/CommonUI/Public/CommonVisualAttachment.h
new file mode 100644
--- /dev/null
+++ b/Engine/Plugins/Runtime/CommonUI/Source/CommonUI/Public/CommonVisualAttachment.h
@@ -0,0 +1,39 @@
+// Copyright Epic Games, Inc. All Rights Reserved.
+
+#pragma once
+
+#include "Components/SizeBox.h"
+
+#include "CommonVisualAttachment.generated.h"
+
+/**
+ * Adds a widget as a zero-size attachment to another. Think icons to the left of labels, without changing the computed size of the label.
+ */
+UCLASS(ClassGroup = UI, meta = (Category = "Common UI"))
+class COMMONUI_API UCommonVisualAttachment : public USizeBox
+{
+	GENERATED_BODY()
+
+public:
+	UCommonVisualAttachment(const FObjectInitializer& ObjectInitializer);
+
+	/** */
+	UPROPERTY(EditAnywhere, BlueprintReadOnly, Category="Child Layout")
+	FVector2D ContentAnchor;
+
+protected:
+	// UVisual interface
+	virtual void ReleaseSlateResources(bool bReleaseChildren) override;
+	virtual TSharedRef<SWidget> RebuildWidget() override;
+	virtual void SynchronizeProperties() override;
+	// End of UWidget interface
+
+private:
+	TSharedPtr<class SVisualAttachmentBox> MyAttachmentBox;
+};
+
+#if UE_ENABLE_INCLUDE_ORDER_DEPRECATED_IN_5_2
+#include "CoreMinimal.h"
+#include "UObject/ObjectMacros.h"
+#include "Widgets/SWidget.h"
+#endif
diff --git a/Engine/Plugins/Runtime/CommonUI/Source/CommonUI/Public/CommonWidgetCarousel.h b/Engine/Plugins/Runtime/CommonUI/Source/CommonUI/Public/CommonWidgetCarousel.h
new file mode 100644
--- /dev/null
+++ b/Engine/Plugins/Runtime/CommonUI/Source/CommonUI/Public/CommonWidgetCarousel.h
@@ -0,0 +1,99 @@
+// Copyright Epic Games, Inc. All Rights Reserved.
+
+#pragma once
+
+#include "Components/PanelWidget.h"
+#include "SWidgetCarousel.h"
+
+class UCommonWidgetCarousel;
+
+#include "CommonWidgetCarousel.generated.h"
+
+enum EHorizontalAlignment : int;
+enum EVerticalAlignment : int;
+
+DECLARE_DYNAMIC_MULTICAST_DELEGATE_TwoParams(FOnCurrentPageIndexChanged, UCommonWidgetCarousel*, CarouselWidget, int32, CurrentPageIndex);
+
+/**
+ * A widget switcher is like a tab control, but without tabs. At most one widget is visible at time.
+ */
+UCLASS(Blueprintable)
+class COMMONUI_API UCommonWidgetCarousel : public UPanelWidget
+{
+	GENERATED_UCLASS_BODY()
+
+public:
+	/** The slot index to display */
+	UPROPERTY(EditAnywhere, BlueprintReadOnly, Category="Carousel", meta=( UIMin=0, ClampMin=0 ))
+	int32 ActiveWidgetIndex;
+
+public:
+
+	/** Gets the slot index of the currently active widget */
+	UFUNCTION(BlueprintCallable, Category="Carousel")
+	int32 GetActiveWidgetIndex() const;
+
+	/** Activates the widget at the specified index. */
+	UFUNCTION(BlueprintCallable, Category="Carousel")
+	virtual void SetActiveWidgetIndex( int32 Index );
+
+	/** Activates the widget and makes it the active index. */
+	UFUNCTION(BlueprintCallable, Category="Carousel")
+	virtual void SetActiveWidget(UWidget* Widget);
+
+	/** Get a widget at the provided index */
+	UFUNCTION( BlueprintCallable, Category = "Carousel" )
+	UWidget* GetWidgetAtIndex( int32 Index ) const;
+
+	UFUNCTION(BlueprintCallable, Category="Carousel")
+	void BeginAutoScrolling(float ScrollInterval = 10);
+
+	UFUNCTION(BlueprintCallable, Category="Carousel")
+	void EndAutoScrolling();
+
+	UFUNCTION(BlueprintCallable, Category="Carousel")
+	void NextPage();
+
+	UFUNCTION(BlueprintCallable, Category="Carousel")
+	void PreviousPage();
+	
+	UPROPERTY(BlueprintAssignable, Category = "Carousel")
+	FOnCurrentPageIndexChanged OnCurrentPageIndexChanged;
+
+	// UWidget interface
+	virtual void SynchronizeProperties() override;
+	virtual void ReleaseSlateResources(bool bReleaseChildren) override;
+	// End of UWidget interface
+
+#if WITH_EDITOR
+	virtual const FText GetPaletteCategory() override;
+	virtual void OnDescendantSelectedByDesigner(UWidget* DescendantWidget) override;
+	virtual void OnDescendantDeselectedByDesigner(UWidget* DescendantWidget) override;
+
+	// UObject interface
+	virtual void PostEditChangeProperty(struct FPropertyChangedEvent& PropertyChangedEvent) override;
+	// End of UObject interface
+#endif
+
+protected:
+	bool AutoScrollCallback(float DeltaTime);
+
+	TSharedRef<SWidget> OnGenerateWidgetForCarousel(UPanelSlot* PanelSlot);
+	void HandlePageChanged(int32 PageIndex);
+
+	// UPanelWidget
+	virtual UClass* GetSlotClass() const override;
+	virtual void OnSlotAdded(UPanelSlot* InSlot) override;
+	// End UPanelWidget
+
+	// UWidget interface
+	virtual TSharedRef<SWidget> RebuildWidget() override;
+	// End of UWidget interface
+
+protected:
+	FTSTicker::FDelegateHandle TickerHandle;
+
+	TSharedPtr< SWidgetCarousel<UPanelSlot*> > MyCommonWidgetCarousel;
+
+	TArray< TSharedRef<SWidget> > CachedSlotWidgets;
+};
diff --git a/Engine/Plugins/Runtime/CommonUI/Source/CommonUI/Public/CommonWidgetCarouselNavBar.h b/Engine/Plugins/Runtime/CommonUI/Source/CommonUI/Public/CommonWidgetCarouselNavBar.h
new file mode 100644
--- /dev/null
+++ b/Engine/Plugins/Runtime/CommonUI/Source/CommonUI/Public/CommonWidgetCarouselNavBar.h
@@ -0,0 +1,73 @@
+// Copyright Epic Games, Inc. All Rights Reserved.
+
+#pragma once
+
+#include "Components/Widget.h"
+#include "CommonWidgetCarouselNavBar.generated.h"
+
+class UCommonWidgetCarousel;
+class UCommonButtonBase;
+class UCommonButtonGroupBase;
+class SHorizontalBox;
+
+/**
+ * A Navigation control for a Carousel
+ */
+UCLASS(Blueprintable)
+class COMMONUI_API UCommonWidgetCarouselNavBar : public UWidget
+{
+	GENERATED_UCLASS_BODY()
+
+public:
+
+	UPROPERTY(EditAnywhere, BlueprintReadOnly, Category = "CarouselNavBar")
+	TSubclassOf<UCommonButtonBase> ButtonWidgetType;
+
+	UPROPERTY(EditAnywhere, BlueprintReadOnly, Category = "CarouselNavBar")
+	FMargin ButtonPadding;
+
+	/**
+	 * Establishes the Widget Carousel instance that this Nav Bar should interact with
+	 * @param CommonCarousel The carousel that this nav bar should be associated with and manipulate
+	 */
+	UFUNCTION(BlueprintCallable, Category = "CarouselNavBar")
+	void SetLinkedCarousel(UCommonWidgetCarousel* CommonCarousel);
+
+
+	// UWidget interface
+	virtual void ReleaseSlateResources(bool bReleaseChildren) override;
+	// End of UWidget interface
+
+#if WITH_EDITOR
+	virtual const FText GetPaletteCategory() override;
+#endif
+
+protected:
+	UFUNCTION()
+	void HandlePageChanged(UCommonWidgetCarousel* CommonCarousel, int32 PageIndex);
+
+	UFUNCTION()
+	void HandleButtonClicked(UCommonButtonBase* AssociatedButton, int32 ButtonIndex);
+
+	// UWidget interface
+	virtual TSharedRef<SWidget> RebuildWidget() override;
+	// End of UWidget interface
+
+	void RebuildButtons();
+
+protected:
+	TSharedPtr<SHorizontalBox> MyContainer;
+	
+	UPROPERTY()
+	TObjectPtr<UCommonWidgetCarousel> LinkedCarousel;
+
+	UPROPERTY()
+	TObjectPtr<UCommonButtonGroupBase> ButtonGroup;
+
+	UPROPERTY()
+	TArray<TObjectPtr<UCommonButtonBase>> Buttons;
+};
+
+#if UE_ENABLE_INCLUDE_ORDER_DEPRECATED_IN_5_2
+#include "Components/PanelWidget.h"
+#endif
