Move CommonUI out of Experimental and into Beta

diff --git a/Engine/Plugins/Runtime/CommonUI/Source/CommonUI/Public/Input/CommonAnalogCursor.h b/Engine/Plugins/Runtime/CommonUI/Source/CommonUI/Public/Input/CommonAnalogCursor.h
new file mode 100644
--- /dev/null
+++ b/Engine/Plugins/Runtime/CommonUI/Source/CommonUI/Public/Input/CommonAnalogCursor.h
@@ -0,0 +1,122 @@
+// Copyright Epic Games, Inc. All Rights Reserved.
+
+#pragma once
+
+#include "Framework/Application/AnalogCursor.h"
+#include "InputCoreTypes.h"
+#include "Rendering/SlateRenderTransform.h"
+
+class UCommonUIActionRouterBase;
+class UCommonInputSubsystem;
+class SWidget;
+class UWidget;
+class UGameViewportClient;
+
+struct FInputEvent;
+enum class ECommonInputType : uint8;
+enum EOrientation : int;
+
+/**
+ * Analog cursor preprocessor that tastefully hijacks things a bit to support controller navigation by moving a hidden cursor around based on focus.
+ *
+ * Introduces a separate focus-driven mode of operation, wherein the cursor is made invisible and automatically updated
+ * to be centered over whatever widget is currently focused (except the game viewport - we completely hide it then)
+ */
+class COMMONUI_API FCommonAnalogCursor : public FAnalogCursor
+{
+public:
+	template <typename AnalogCursorT = FCommonAnalogCursor>
+	static TSharedRef<AnalogCursorT> CreateAnalogCursor(const UCommonUIActionRouterBase& InActionRouter)
+	{
+		TSharedRef<AnalogCursorT> NewCursor = MakeShareable(new AnalogCursorT(InActionRouter));
+		NewCursor->Initialize();
+		return NewCursor;
+	}
+
+	virtual void Tick(const float DeltaTime, FSlateApplication& SlateApp, TSharedRef<ICursor> Cursor) override;
+	virtual bool HandleKeyDownEvent(FSlateApplication& SlateApp, const FKeyEvent& InKeyEvent) override;
+	virtual bool HandleKeyUpEvent(FSlateApplication& SlateApp, const FKeyEvent& InKeyEvent) override;
+
+	virtual bool CanReleaseMouseCapture() const;
+
+	virtual bool HandleAnalogInputEvent(FSlateApplication& SlateApp, const FAnalogInputEvent& InAnalogInputEvent) override;
+	virtual bool HandleMouseMoveEvent(FSlateApplication& SlateApp, const FPointerEvent& MouseEvent) override;
+	virtual bool HandleMouseButtonDownEvent(FSlateApplication& SlateApp, const FPointerEvent& InPointerEvent) override;
+	virtual bool HandleMouseButtonUpEvent(FSlateApplication& SlateApp, const FPointerEvent& InPointerEvent) override;
+
+	void SetCursorMovementStick(EAnalogStick InCursorMovementStick);
+
+	virtual int32 GetOwnerUserIndex() const override;
+
+	virtual void ShouldHandleRightAnalog(bool bInShouldHandleRightAnalog);
+
+	virtual bool IsAnalogMovementEnabled() const { return bIsAnalogMovementEnabled; }
+
+protected:
+	FCommonAnalogCursor(const UCommonUIActionRouterBase& InActionRouter);
+	virtual void Initialize();
+	
+	virtual EOrientation DetermineScrollOrientation(const UWidget& Widget) const;
+
+	virtual bool IsRelevantInput(const FKeyEvent& KeyEvent) const override;
+	virtual bool IsRelevantInput(const FAnalogInputEvent& AnalogInputEvent) const override;
+	
+	void SetNormalizedCursorPosition(const FVector2D& RelativeNewPosition);
+	bool IsInViewport(const FVector2D& Position) const;
+	FVector2D ClampPositionToViewport(const FVector2D& InPosition) const;
+	void HideCursor();
+
+	UGameViewportClient* GetViewportClient() const;
+	
+	/**
+	 * A ridiculous function name, but we have this exact question in a few places.
+	 * We don't care about input while our owning player's game viewport isn't involved in the focus path,
+	 * but we also want to hold off doing anything while that game viewport has full capture.
+	 * So we need that "relevant, but not exclusive" sweet spot.
+	 */
+	bool IsGameViewportInFocusPathWithoutCapture() const;
+
+	virtual void RefreshCursorSettings();
+	virtual void RefreshCursorVisibility();
+	
+	virtual void HandleInputMethodChanged(ECommonInputType NewInputMethod);
+	
+	bool IsUsingGamepad() const;
+
+	bool ShouldHideCursor() const;
+
+	// Knowingly unorthodox member reference to a UObject - ok because we are a subobject of the owning router and will never outlive it
+	const UCommonUIActionRouterBase& ActionRouter;
+	ECommonInputType ActiveInputMethod;
+
+	bool bIsAnalogMovementEnabled = false;
+	
+	bool bShouldHandleRightAnalog = true;
+	
+private:
+
+	/** The current set of pointer buttons being used as keys. */
+	TSet<FKey> PointerButtonDownKeys;
+
+	TWeakPtr<SWidget> LastCursorTarget;
+	FSlateRenderTransform LastCursorTargetTransform;
+
+	float TimeUntilScrollUpdate = 0.f;
+	
+#if !UE_BUILD_SHIPPING
+	enum EShoulderButtonFlags
+	{
+		None = 0,
+		LeftShoulder = 1 << 0,
+		RightShoulder = 1 << 1,
+		LeftTrigger = 1 << 2,
+		RightTrigger = 1 << 3,
+		All = 0x1111
+	};
+	int32 ShoulderButtonStatus = EShoulderButtonFlags::None;
+#endif
+};
+
+#if UE_ENABLE_INCLUDE_ORDER_DEPRECATED_IN_5_2
+#include "Layout/Geometry.h"
+#endif
diff --git a/Engine/Plugins/Runtime/CommonUI/Source/CommonUI/Public/Input/CommonBoundActionBar.h b/Engine/Plugins/Runtime/CommonUI/Source/CommonUI/Public/Input/CommonBoundActionBar.h
new file mode 100644
--- /dev/null
+++ b/Engine/Plugins/Runtime/CommonUI/Source/CommonUI/Public/Input/CommonBoundActionBar.h
@@ -0,0 +1,63 @@
+// Copyright Epic Games, Inc. All Rights Reserved.
+
+#pragma once
+
+#include "CommonButtonBase.h"
+#include "Components/DynamicEntryBoxBase.h"
+#include "Tickable.h"
+#include "CommonBoundActionBar.generated.h"
+
+class ICommonBoundActionButtonInterface;
+class IConsoleVariable;
+class IWidgetCompilerLog;
+
+struct FUIActionBindingHandle;
+
+/**
+ * A box populated with current actions available per CommonUI's Input Handler.
+ */
+UCLASS(Blueprintable, ClassGroup = UI, meta = (Category = "Common UI"))
+class COMMONUI_API UCommonBoundActionBar : public UDynamicEntryBoxBase, public FTickableGameObject
+{
+	GENERATED_BODY()
+
+public:
+	UFUNCTION(BlueprintCallable, Category = CommonBoundActionBar)
+	void SetDisplayOwningPlayerActionsOnly(bool bShouldOnlyDisplayOwningPlayerActions);
+
+	//~ FTickableGameObject Begin
+	virtual void Tick(float DeltaTime) override;
+	virtual ETickableTickType GetTickableTickType() const override;
+	virtual TStatId GetStatId() const override;
+	virtual bool IsTickableWhenPaused() const override;
+	//~ FTickableGameObject End
+
+protected:
+	virtual void OnWidgetRebuilt() override;
+	virtual void SynchronizeProperties() override;
+	virtual void ReleaseSlateResources(bool bReleaseChildren) override;
+
+	virtual void NativeOnActionButtonCreated(ICommonBoundActionButtonInterface* ActionButton, const FUIActionBindingHandle& RepresentedAction) { }
+
+#if WITH_EDITOR
+	void ValidateCompiledDefaults(IWidgetCompilerLog& CompileLog) const override;
+#endif
+
+private:
+	void HandleBoundActionsUpdated(bool bFromOwningPlayer);
+	void HandleDeferredDisplayUpdate();
+	void HandlePlayerAdded(int32 PlayerIdx);
+	
+	void MonitorPlayerActions(const ULocalPlayer* NewPlayer);
+
+	UPROPERTY(EditAnywhere, Category = EntryLayout, meta=(MustImplement = "/Script/CommonUI.CommonBoundActionButtonInterface"))
+	TSubclassOf<UCommonButtonBase> ActionButtonClass;
+
+	UPROPERTY(EditAnywhere, Category = Display)
+	bool bDisplayOwningPlayerActionsOnly = true;
+
+	UPROPERTY(EditAnywhere, AdvancedDisplay, Category = Display)
+	bool bIgnoreDuplicateActions = true;
+
+	bool bIsRefreshQueued = false;
+};
\ No newline at end of file
diff --git a/Engine/Plugins/Runtime/CommonUI/Source/CommonUI/Public/Input/CommonBoundActionButton.h b/Engine/Plugins/Runtime/CommonUI/Source/CommonUI/Public/Input/CommonBoundActionButton.h
new file mode 100644
--- /dev/null
+++ b/Engine/Plugins/Runtime/CommonUI/Source/CommonUI/Public/Input/CommonBoundActionButton.h
@@ -0,0 +1,37 @@
+// Copyright Epic Games, Inc. All Rights Reserved.
+
+#pragma once
+
+#include "CommonBoundActionButtonInterface.h"
+#include "CommonButtonBase.h"
+#include "UIActionBindingHandle.h"
+#include "CommonBoundActionButton.generated.h"
+
+class UCommonTextBlock;
+
+UCLASS(Abstract, meta = (DisableNativeTick))
+class COMMONUI_API UCommonBoundActionButton : public UCommonButtonBase, public ICommonBoundActionButtonInterface
+{
+	GENERATED_BODY()
+
+public:
+	//~ Begin ICommonBoundActionButtonInterface
+	virtual void SetRepresentedAction(FUIActionBindingHandle InBindingHandle) override;
+	//~ End ICommonBoundActionButtonInterface
+	
+protected:
+	virtual void NativeOnClicked() override;
+	virtual void NativeOnCurrentTextStyleChanged() override;
+
+	virtual void UpdateInputActionWidget() override;
+
+	UFUNCTION(BlueprintImplementableEvent, Category = "Common Bound Action")
+	void OnUpdateInputAction();
+
+protected:
+	UPROPERTY(BlueprintReadOnly, meta = (BindWidget), Category = "Text Block")
+	TObjectPtr<UCommonTextBlock> Text_ActionName;
+
+private:
+	FUIActionBindingHandle BindingHandle;
+};
\ No newline at end of file
diff --git a/Engine/Plugins/Runtime/CommonUI/Source/CommonUI/Public/Input/CommonUIActionRouterBase.h b/Engine/Plugins/Runtime/CommonUI/Source/CommonUI/Public/Input/CommonUIActionRouterBase.h
new file mode 100644
--- /dev/null
+++ b/Engine/Plugins/Runtime/CommonUI/Source/CommonUI/Public/Input/CommonUIActionRouterBase.h
@@ -0,0 +1,249 @@
+// Copyright Epic Games, Inc. All Rights Reserved.
+
+#pragma once
+
+#include "CommonInputModeTypes.h"
+#include "Subsystems/LocalPlayerSubsystem.h"
+#include "Containers/Ticker.h"
+#include "Containers/CircularBuffer.h"
+
+#include "Engine/EngineBaseTypes.h"
+#include "Input/UIActionBindingHandle.h"
+#include "InputCoreTypes.h"
+#include "UObject/WeakObjectPtr.h"
+#include "CommonUIActionRouterBase.generated.h"
+
+class SWidget;
+struct FBindUIActionArgs;
+
+class AHUD;
+class FWeakWidgetPath;
+class FWidgetPath;
+class UCanvas;
+class UWidget;
+class UCommonUserWidget;
+class UCommonActivatableWidget;
+class UCommonInputSubsystem;
+class UInputComponent;
+class UPlayerInput;
+class FCommonAnalogCursor;
+class IInputProcessor;
+
+enum class EProcessHoldActionResult;
+class FActivatableTreeNode;
+class UCommonInputSubsystem;
+class UCommonInputActionDomain;
+struct FUIActionBinding;
+class FDebugDisplayInfo;
+struct FAutoCompleteCommand;
+using FActivatableTreeNodePtr = TSharedPtr<FActivatableTreeNode>;
+using FActivatableTreeNodeRef = TSharedRef<FActivatableTreeNode>;
+
+class FActivatableTreeRoot;
+using FActivatableTreeRootPtr = TSharedPtr<FActivatableTreeRoot>;
+using FActivatableTreeRootRef = TSharedRef<FActivatableTreeRoot>;
+struct FFocusEvent;
+
+enum class ERouteUIInputResult : uint8
+{
+	Handled,
+	BlockGameInput,
+	Unhandled
+};
+
+/**
+ * The nucleus of the CommonUI input routing system. 
+ * 
+ * Gathers input from external sources such as game viewport client and forwards them to widgets 
+ * via activatable tree node representation.
+ */
+UCLASS()
+class COMMONUI_API UCommonUIActionRouterBase : public ULocalPlayerSubsystem
+{
+	GENERATED_BODY()
+
+public:
+	static UCommonUIActionRouterBase* Get(const UWidget& ContextWidget);
+
+	/** searches up the SWidget tree until it finds the nearest UCommonActivatableWidget */
+	static UCommonActivatableWidget* FindOwningActivatable(TSharedPtr<SWidget> Widget, ULocalPlayer* OwningLocalPlayer);
+
+	UCommonUIActionRouterBase();
+	virtual void Initialize(FSubsystemCollectionBase& Collection) override;
+	virtual void Deinitialize() override;
+	virtual bool ShouldCreateSubsystem(UObject* Outer) const override;
+
+	/** Sets whether the underlying activatable tree system is enabled - when disabled, all we really do is process Persistent input actions */
+	virtual void SetIsActivatableTreeEnabled(bool bInIsTreeEnabled);
+
+	virtual FUIActionBindingHandle RegisterUIActionBinding(const UWidget& Widget, const FBindUIActionArgs& BindActionArgs);
+	bool RegisterLinkedPreprocessor(const UWidget& Widget, const TSharedRef<IInputProcessor>& InputPreprocessor, int32 DesiredIndex = INDEX_NONE);
+
+	DECLARE_EVENT_OneParam(UCommonUIActionRouterBase, FOnActiveInputModeChanged, ECommonInputMode);
+	FOnActiveInputModeChanged& OnActiveInputModeChanged() const { return OnActiveInputModeChangedEvent; }
+	ECommonInputMode GetActiveInputMode(ECommonInputMode DefaultInputMode = ECommonInputMode::All) const;
+	EMouseCaptureMode GetActiveMouseCaptureMode(EMouseCaptureMode DefaultMouseCapture = EMouseCaptureMode::NoCapture) const;
+
+	DECLARE_EVENT_OneParam(UCommonUIActionRouterBase, FOnActivationMetadataChanged, FActivationMetadata);
+	FOnActivationMetadataChanged& OnActivationMetadataChanged() const { return OnActivationMetadataChangedEvent; }
+
+	void RegisterScrollRecipient(const UWidget& ScrollableWidget);
+	void UnregisterScrollRecipient(const UWidget& ScrollableWidget);
+	TArray<const UWidget*> GatherActiveAnalogScrollRecipients() const;
+
+	TArray<FUIActionBindingHandle> GatherActiveBindings() const;
+	FSimpleMulticastDelegate& OnBoundActionsUpdated() const { return OnBoundActionsUpdatedEvent; }
+
+	UCommonInputSubsystem& GetInputSubsystem() const;
+
+	ERouteUIInputResult ProcessInput(FKey Key, EInputEvent InputEvent) const;
+	bool CanProcessNormalGameInput() const;
+
+	bool IsPendingTreeChange() const;
+
+	TSharedPtr<FCommonAnalogCursor> GetCommonAnalogCursor() const { return AnalogCursor; }
+
+	void FlushInput();
+
+	bool IsWidgetInActiveRoot(const UCommonActivatableWidget* Widget) const;
+
+	/** 
+	 * Sets Input Config 
+	 * 
+	 * @param NewConfig config to set
+	 * @param InConfigSource optional source of config. If exists, will be used to log input config source
+	 */
+	void SetActiveUIInputConfig(const FUIInputConfig& NewConfig, const UObject* InConfigSource = nullptr);
+
+public:
+	void NotifyUserWidgetConstructed(const UCommonUserWidget& Widget);
+	void NotifyUserWidgetDestructed(const UCommonUserWidget& Widget);
+	
+	void AddBinding(FUIActionBindingHandle Binding);
+	void RemoveBinding(FUIActionBindingHandle Binding);
+
+	int32 GetLocalPlayerIndex() const;
+
+	void RefreshActiveRootFocus();
+	void RefreshUIInputConfig();
+
+	bool ShouldAlwaysShowCursor() const;
+
+protected:
+	virtual TSharedRef<FCommonAnalogCursor> MakeAnalogCursor() const;
+	virtual void PostAnalogCursorCreate();
+	void RegisterAnalogCursorTick();
+
+	virtual void SetActiveRoot(FActivatableTreeRootPtr NewActiveRoot);
+	void SetForceResetActiveRoot(bool bInForceResetActiveRoot);
+
+	virtual void ApplyUIInputConfig(const FUIInputConfig& NewConfig, bool bForceRefresh);
+	void UpdateLeafNodeAndConfig(FActivatableTreeRootPtr DesiredRoot, FActivatableTreeNodePtr DesiredLeafNode);
+
+	void RefreshActionDomainLeafNodeConfig();
+
+	bool bIsActivatableTreeEnabled = true;
+
+	/** The currently applied UI input configuration */
+	TOptional<FUIInputConfig> ActiveInputConfig;
+
+	TSharedPtr<FCommonAnalogCursor> AnalogCursor;
+	FTSTicker::FDelegateHandle TickHandle;
+
+private:
+	bool Tick(float DeltaTime);
+
+	void OnShowDebugInfo(AHUD* HUD, UCanvas* Canvas, const FDebugDisplayInfo& DisplayInfo, float& YL, float& YPos);
+	void PopulateAutoCompleteEntries(TArray<FAutoCompleteCommand>& AutoCompleteList);
+
+	void RegisterWidgetBindings(const FActivatableTreeNodePtr& TreeNode, const TArray<FUIActionBindingHandle>& BindingHandles);
+
+	FActivatableTreeNodePtr FindNode(const UCommonActivatableWidget* Widget) const;
+	FActivatableTreeNodePtr FindOwningNode(const UWidget& Widget) const;
+	FActivatableTreeNodePtr FindNodeRecursive(const FActivatableTreeNodePtr& CurrentNode, const UCommonActivatableWidget& Widget) const;
+	FActivatableTreeNodePtr FindNodeRecursive(const FActivatableTreeNodePtr& CurrentNode, const TSharedPtr<SWidget>& Widget) const;
+	void SetActiveActivationMetadata(const FActivationMetadata& NewConfig);
+	
+	void HandleActivatableWidgetRebuilding(UCommonActivatableWidget& RebuildingWidget);
+	void ProcessRebuiltWidgets();
+	void AssembleTreeRecursive(const FActivatableTreeNodeRef& CurNode, TMap<UCommonActivatableWidget*, TArray<UCommonActivatableWidget*>>& WidgetsByDirectParent);
+
+	void HandleRootWidgetSlateReleased(TWeakPtr<FActivatableTreeRoot> WeakRoot);
+	void HandleRootNodeActivated(TWeakPtr<FActivatableTreeRoot> WeakActivatedRoot);
+	void HandleRootNodeDeactivated(TWeakPtr<FActivatableTreeRoot> WeakDeactivatedRoot);
+	void HandleLeafmostActiveNodeChanged();
+
+	void HandleSlateFocusChanging(const FFocusEvent& FocusEvent, const FWeakWidgetPath& OldFocusedWidgetPath, const TSharedPtr<SWidget>& OldFocusedWidget, const FWidgetPath& NewFocusedWidgetPath, const TSharedPtr<SWidget>& NewFocusedWidget);
+
+	void HandlePostGarbageCollect();
+
+	bool ProcessInputOnActionDomains(ECommonInputMode ActiveInputMode, FKey Key, EInputEvent InputEvent) const;
+	EProcessHoldActionResult ProcessHoldInputOnActionDomains(ECommonInputMode ActiveInputMode, FKey Key, EInputEvent InputEvent) const;
+
+	struct FPendingWidgetRegistration
+	{
+		TWeakObjectPtr<const UWidget> Widget;
+		TArray<FUIActionBindingHandle> ActionBindings;
+		bool bIsScrollRecipient = false;
+		
+		struct FPreprocessorRegistration
+		{
+			TSharedPtr<IInputProcessor> Preprocessor;
+			int32 DesiredIdx = 0;
+			bool operator==(const TSharedRef<IInputProcessor>& OtherPreprocessor) const { return Preprocessor == OtherPreprocessor; }
+		};
+		TArray<FPreprocessorRegistration> Preprocessors;
+
+		bool operator==(const UWidget* OtherWidget) const { return OtherWidget == Widget.Get(); }
+		bool operator==(const UWidget& OtherWidget) const { return &OtherWidget == Widget.Get(); }
+	};
+	FPendingWidgetRegistration& GetOrCreatePendingRegistration(const UWidget& Widget);
+	TArray<FPendingWidgetRegistration> PendingWidgetRegistrations;
+	
+	TArray<TWeakObjectPtr<UCommonActivatableWidget>> RebuiltWidgetsPendingNodeAssignment;
+	
+	TArray<FActivatableTreeRootRef> RootNodes;
+	FActivatableTreeRootPtr ActiveRootNode;
+
+	// Note: Treat this as a TSharedRef - only reason it isn't is because TSharedRef doesn't play nice with forward declarations :(
+	TSharedPtr<class FPersistentActionCollection> PersistentActions;
+
+	TCircularBuffer<FString> InputConfigSources = TCircularBuffer<FString>(5, "None");
+	int32 InputConfigSourceIndex = 0;
+
+	bool bForceResetActiveRoot = false;
+
+	mutable FSimpleMulticastDelegate OnBoundActionsUpdatedEvent;
+	mutable FOnActiveInputModeChanged OnActiveInputModeChangedEvent;
+	mutable FOnActivationMetadataChanged OnActivationMetadataChangedEvent;
+
+	friend class FActionRouterBindingCollection;
+	friend class FActivatableTreeNode;
+	friend class FActivatableTreeRoot;
+	friend class FActionRouterDebugUtils;
+
+	mutable TArray<FKey> HeldKeys;
+
+	/** A wrapper around TArray that keeps RootList sorted by PaintLayer during insertion */
+	struct FActionDomainSortedRootList
+	{
+		TArray<FActivatableTreeRootRef> RootList;
+
+		// Inserts RootNode into RootList based on its paint layer
+		void Add(FActivatableTreeRootRef RootNode);
+
+		// Trivial removal
+		int32 Remove(FActivatableTreeRootRef RootNode);
+
+		// Trivial Contains check
+		bool Contains(FActivatableTreeRootRef RootNode) const;
+	};
+
+	TMap<TObjectPtr<UCommonInputActionDomain>, FActionDomainSortedRootList> ActionDomainRootNodes;
+};
+
+#if UE_ENABLE_INCLUDE_ORDER_DEPRECATED_IN_5_2
+#include "CommonInputBaseTypes.h"
+#include "CommonUIInputTypes.h"
+#include "Misc/Passkey.h"
+#endif
diff --git a/Engine/Plugins/Runtime/CommonUI/Source/CommonUI/Public/SVisualAttachmentBox.h b/Engine/Plugins/Runtime/CommonUI/Source/CommonUI/Public/SVisualAttachmentBox.h
new file mode 100644
--- /dev/null
+++ b/Engine/Plugins/Runtime/CommonUI/Source/CommonUI/Public/SVisualAttachmentBox.h
@@ -0,0 +1,29 @@
+// Copyright Epic Games, Inc. All Rights Reserved.
+
+#pragma once
+
+#include "Widgets/Layout/SBox.h"
+
+/**
+ * Wrapper widget meant to handle native-side painting for UCommonVisualAttachment.
+ */
+class COMMONUI_API SVisualAttachmentBox : public SBox
+{
+public:
+	void Construct(const FArguments& InArgs);
+
+	virtual FVector2D ComputeDesiredSize(float) const override;
+	virtual void OnArrangeChildren(const FGeometry& AllottedGeometry, FArrangedChildren& ArrangedChildren) const override;
+	virtual int32 OnPaint(const FPaintArgs& Args, const FGeometry& AllottedGeometry, const FSlateRect& MyCullingRect, FSlateWindowElementList& OutDrawElements, int32 LayerId, const FWidgetStyle& InWidgetStyle, bool bParentEnabled) const override;
+
+	void SetContentAnchor(FVector2D InContentAnchor);
+
+private:
+	mutable FVector2D InnerDesiredSize;
+
+	FVector2D ContentAnchor = FVector2D::ZeroVector;
+};
+
+#if UE_ENABLE_INCLUDE_ORDER_DEPRECATED_IN_5_2
+#include "CoreMinimal.h"
+#endif
