Fix CommonRichtextBlock where the text runs are getting scaled up twice on mobile
#RNX

diff --git a/Engine/Plugins/Runtime/CommonUI/Source/CommonUI/Private/CommonRichTextBlock.cpp b/Engine/Plugins/Runtime/CommonUI/Source/CommonUI/Private/CommonRichTextBlock.cpp
--- a/Engine/Plugins/Runtime/CommonUI/Source/CommonUI/Private/CommonRichTextBlock.cpp
+++ b/Engine/Plugins/Runtime/CommonUI/Source/CommonUI/Private/CommonRichTextBlock.cpp
@@ -24,6 +24,16 @@
 
 #include UE_INLINE_GENERATED_CPP_BY_NAME(CommonRichTextBlock)
 
+namespace UE::CommonUI::CommonRichTextBlock
+{
+	int32 bApplyMobileScaleOnTextRuns = 0;
+	FAutoConsoleVariableRef CVarApplyMobileScaleOnTextRuns(
+		TEXT("CommonRichTextBlock.CVarApplyMobileScaleOnTextRuns"),
+		bApplyMobileScaleOnTextRuns,
+		TEXT("Enable the mobile scale on the text when it gets created. It is off by default because it can result its scale being applied twice in tagged text.\n0: Disable (default), 1: Enable"),
+		ECVF_Default);
+}
+
 namespace SupportedMarkupKeys
 {
 	const FString Id = TEXT("id");
@@ -114,7 +124,7 @@ public:
 			}
 		}
 
-		if (CommonUIUtils::ShouldDisplayMobileUISizes())
+		if (CommonUIUtils::ShouldDisplayMobileUISizes() && UE::CommonUI::CommonRichTextBlock::bApplyMobileScaleOnTextRuns)
 		{
 			TextStyle.SetFontSize(TextStyle.Font.Size * Owner->GetMobileTextBlockScale());
 		}
