Modify Hold Action Binding Pressed and Released behavior on Common Button Base
- Added HoldActionReleased delegates to UI Input Actions
- Added broadcast logic to the released portion
- NativeOnPressed and NativeOnReleased execute correspondingly

diff --git a/Engine/Plugins/Runtime/CommonUI/Source/CommonUI/Private/CommonButtonBase.cpp b/Engine/Plugins/Runtime/CommonUI/Source/CommonUI/Private/CommonButtonBase.cpp
--- a/Engine/Plugins/Runtime/CommonUI/Source/CommonUI/Private/CommonButtonBase.cpp
+++ b/Engine/Plugins/Runtime/CommonUI/Source/CommonUI/Private/CommonButtonBase.cpp
@@ -602,7 +602,8 @@ void UCommonButtonBase::BindTriggeringInputActionToClick()
 	{
 		FBindUIActionArgs BindArgs(TriggeringInputAction, false, FSimpleDelegate::CreateUObject(this, &UCommonButtonBase::HandleTriggeringActionCommited));
 		BindArgs.OnHoldActionProgressed.BindUObject(this, &UCommonButtonBase::NativeOnActionProgress);
-		BindArgs.OnHoldActionPressed.BindUObject(this, &UCommonButtonBase::HoldReset);
+		BindArgs.OnHoldActionPressed.BindUObject(this, &UCommonButtonBase::NativeOnPressed);
+		BindArgs.OnHoldActionReleased.BindUObject(this, &UCommonButtonBase::NativeOnReleased);
 		BindArgs.bIsPersistent = bIsPersistentBinding;
 		BindArgs.bForceHold = GetConvertInputActionToHold();
 
@@ -1525,6 +1526,7 @@ void UCommonButtonBase::StopDoubleClickPropagation()
 
 void UCommonButtonBase::NativeOnPressed()
 {
+	HoldReset();
 	BP_OnPressed();
 	OnPressed().Broadcast();
 }
diff --git a/Engine/Plugins/Runtime/CommonUI/Source/CommonUI/Private/Input/UIActionRouterTypes.cpp b/Engine/Plugins/Runtime/CommonUI/Source/CommonUI/Private/Input/UIActionRouterTypes.cpp
--- a/Engine/Plugins/Runtime/CommonUI/Source/CommonUI/Private/Input/UIActionRouterTypes.cpp
+++ b/Engine/Plugins/Runtime/CommonUI/Source/CommonUI/Private/Input/UIActionRouterTypes.cpp
@@ -120,6 +120,7 @@ FUIActionBinding::FUIActionBinding(const UWidget& InBoundWidget, const FBindUIAc
 {
 	OnHoldActionProgressed.Add(BindArgs.OnHoldActionProgressed);
 	OnHoldActionPressed.Add(BindArgs.OnHoldActionPressed);
+	OnHoldActionReleased.Add(BindArgs.OnHoldActionReleased);
 
 	const auto RegisterKeyMappingFunc = 
 		[this](const FUIActionKeyMapping& KeyMapping)
@@ -606,6 +607,8 @@ EProcessHoldActionResult FActionRouterBindingCollection::ProcessHoldInput(ECommo
 							{
 								const float SecondsHeld = Binding->GetSecondsHeld();
 
+								Binding->OnHoldActionReleased.Broadcast();
+
 								UE_LOG(LogUIActionRouter, VeryVerbose, TEXT("Hold released: %s after %f"), *Binding->ActionName.ToString(), SecondsHeld);
 
 								if (HoldMapping.HoldRollbackTime <= UE_SMALL_NUMBER)
diff --git a/Engine/Plugins/Runtime/CommonUI/Source/CommonUI/Public/Input/CommonUIInputTypes.h b/Engine/Plugins/Runtime/CommonUI/Source/CommonUI/Public/Input/CommonUIInputTypes.h
--- a/Engine/Plugins/Runtime/CommonUI/Source/CommonUI/Public/Input/CommonUIInputTypes.h
+++ b/Engine/Plugins/Runtime/CommonUI/Source/CommonUI/Public/Input/CommonUIInputTypes.h
@@ -90,6 +90,10 @@ struct COMMONUI_API FBindUIActionArgs
 	/** If the bound action has any hold mappings, this will fire when the hold begins. Has no bearing on actual execution and wholly irrelevant for non-hold actions */
 	DECLARE_DELEGATE(FOnHoldActionPressed);
 	FOnHoldActionPressed OnHoldActionPressed;
+
+	/** If the bound action has any hold mappings, this will fire when the hold is interrupted. Has no bearing on actual execution and wholly irrelevant for non-hold actions */
+	DECLARE_DELEGATE(FOnHoldActionReleased);
+	FOnHoldActionPressed OnHoldActionReleased;
 };
 
 #if UE_ENABLE_INCLUDE_ORDER_DEPRECATED_IN_5_2
diff --git a/Engine/Plugins/Runtime/CommonUI/Source/CommonUI/Public/Input/UIActionBinding.h b/Engine/Plugins/Runtime/CommonUI/Source/CommonUI/Public/Input/UIActionBinding.h
--- a/Engine/Plugins/Runtime/CommonUI/Source/CommonUI/Public/Input/UIActionBinding.h
+++ b/Engine/Plugins/Runtime/CommonUI/Source/CommonUI/Public/Input/UIActionBinding.h
@@ -74,6 +74,9 @@ struct COMMONUI_API FUIActionBinding
 	DECLARE_MULTICAST_DELEGATE(FOnHoldActionPressed);
 	FOnHoldActionPressed OnHoldActionPressed;
 
+	DECLARE_MULTICAST_DELEGATE(FOnHoldActionReleased);
+	FOnHoldActionPressed OnHoldActionReleased;
+
 	// @TODO: Rename non-legacy in 5.3. We no longer have any active plans to remove data tables in CommonUI.
 	FDataTableRowHandle LegacyActionTableRow;
 
