AnimatedWidgetSwitcher regression fix: use ActiveWidgetIndex variable directly instead of the getter.

diff --git a/Engine/Plugins/Runtime/CommonUI/Source/CommonUI/Private/CommonAnimatedSwitcher.cpp b/Engine/Plugins/Runtime/CommonUI/Source/CommonUI/Private/CommonAnimatedSwitcher.cpp
--- a/Engine/Plugins/Runtime/CommonUI/Source/CommonUI/Private/CommonAnimatedSwitcher.cpp
+++ b/Engine/Plugins/Runtime/CommonUI/Source/CommonUI/Private/CommonAnimatedSwitcher.cpp
@@ -47,11 +47,12 @@ void UCommonAnimatedSwitcher::SetActiveWidget(UWidget* Widget)
 	SetActiveWidgetIndex_Internal(ChildIndex);
 }
 
+PRAGMA_DISABLE_DEPRECATION_WARNINGS
 void UCommonAnimatedSwitcher::ActivateNextWidget(bool bCanWrap)
 {
 	if (Slots.Num() > 1)
 	{
-		if (GetActiveWidgetIndex() == Slots.Num() - 1)
+		if (ActiveWidgetIndex == Slots.Num() - 1)
 		{
 			if (bCanWrap)
 			{
@@ -60,7 +61,7 @@ void UCommonAnimatedSwitcher::ActivateNextWidget(bool bCanWrap)
 		}
 		else
 		{
-			SetActiveWidgetIndex(GetActiveWidgetIndex() + 1);
+			SetActiveWidgetIndex(ActiveWidgetIndex + 1);
 		}
 	}
 }
@@ -69,7 +70,7 @@ void UCommonAnimatedSwitcher::ActivatePreviousWidget(bool bCanWrap)
 {
 	if (Slots.Num() > 1)
 	{
-		if (GetActiveWidgetIndex() == 0)
+		if (ActiveWidgetIndex == 0)
 		{
 			if (bCanWrap)
 			{
@@ -78,7 +79,7 @@ void UCommonAnimatedSwitcher::ActivatePreviousWidget(bool bCanWrap)
 		}
 		else
 		{
-			SetActiveWidgetIndex(GetActiveWidgetIndex() - 1);
+			SetActiveWidgetIndex(ActiveWidgetIndex - 1);
 		}
 	}
 }
@@ -100,21 +101,24 @@ bool UCommonAnimatedSwitcher::IsCurrentlySwitching() const
 
 void UCommonAnimatedSwitcher::HandleSlateActiveIndexChanged(int32 ActiveIndex)
 {
-	if (Slots.IsValidIndex(GetActiveWidgetIndex()))
+	if (Slots.IsValidIndex(ActiveWidgetIndex))
 	{
-		OnActiveWidgetIndexChanged.Broadcast(GetWidgetAtIndex(GetActiveWidgetIndex()), GetActiveWidgetIndex());
+		OnActiveWidgetIndexChanged.Broadcast(GetWidgetAtIndex(ActiveWidgetIndex), ActiveWidgetIndex);
 	}
 }
+PRAGMA_ENABLE_DEPRECATION_WARNINGS
 
 TSharedRef<SWidget> UCommonAnimatedSwitcher::RebuildWidget()
 {
+PRAGMA_DISABLE_DEPRECATION_WARNINGS
 	MyWidgetSwitcher = MyAnimatedSwitcher = SNew(SCommonAnimatedSwitcher)
-		.InitialIndex(GetActiveWidgetIndex())
+		.InitialIndex(ActiveWidgetIndex)
 		.TransitionCurveType(TransitionCurveType)
 		.TransitionDuration(TransitionDuration)
 		.TransitionType(TransitionType)
 		.OnActiveIndexChanged_UObject(this, &UCommonAnimatedSwitcher::HandleSlateActiveIndexChanged)
 		.OnIsTransitioningChanged_UObject(this, &UCommonAnimatedSwitcher::HandleSlateIsTransitioningChanged);
+PRAGMA_ENABLE_DEPRECATION_WARNINGS
 
 	for (UPanelSlot* CurrentSlot : Slots)
 	{
@@ -156,12 +160,11 @@ void UCommonAnimatedSwitcher::SetActiveWidgetIndex_Internal(int32 Index)
 #endif
 	
 	TGuardValue<bool> bCurrentlySwitchingGuard(bCurrentlySwitching, true);
-
-	if (Index >= 0 && Index < Slots.Num() && (Index != GetActiveWidgetIndex() || !bSetOnce))
+PRAGMA_DISABLE_DEPRECATION_WARNINGS
+	if (Index >= 0 && Index < Slots.Num() && (Index != ActiveWidgetIndex || !bSetOnce))
 	{
 		HandleOutgoingWidget();
 
-PRAGMA_DISABLE_DEPRECATION_WARNINGS
 		// For now we can't call setter since it calls MyWidgetSwitcher->SetActiveWidgetIndex(SafeIndex)
 		if (ActiveWidgetIndex != Index)
 		{
