Fixed typo and change the order of execution in UpdateHold.

diff --git a/Engine/Plugins/Runtime/CommonUI/Source/CommonUI/Public/Input/CommonUIInputTypes.h b/Engine/Plugins/Runtime/CommonUI/Source/CommonUI/Public/Input/CommonUIInputTypes.h
--- a/Engine/Plugins/Runtime/CommonUI/Source/CommonUI/Public/Input/CommonUIInputTypes.h
+++ b/Engine/Plugins/Runtime/CommonUI/Source/CommonUI/Public/Input/CommonUIInputTypes.h
@@ -93,7 +93,7 @@ struct COMMONUI_API FBindUIActionArgs
 
 	/** If the bound action has any hold mappings, this will fire when the hold is interrupted. Has no bearing on actual execution and wholly irrelevant for non-hold actions */
 	DECLARE_DELEGATE(FOnHoldActionReleased);
-	FOnHoldActionPressed OnHoldActionReleased;
+	FOnHoldActionReleased OnHoldActionReleased;
 };
 
 #if UE_ENABLE_INCLUDE_ORDER_DEPRECATED_IN_5_2
