PR #10404: Add COMMONUI_API to FUIInputAction and some related functions in UCommonUIInputSettings.

diff --git a/Engine/Plugins/Runtime/CommonUI/Source/CommonUI/Public/Input/CommonUIInputSettings.h b/Engine/Plugins/Runtime/CommonUI/Source/CommonUI/Public/Input/CommonUIInputSettings.h
--- a/Engine/Plugins/Runtime/CommonUI/Source/CommonUI/Public/Input/CommonUIInputSettings.h
+++ b/Engine/Plugins/Runtime/CommonUI/Source/CommonUI/Public/Input/CommonUIInputSettings.h
@@ -35,7 +35,7 @@ public:
 };
 
 USTRUCT()
-struct FUIInputAction
+struct COMMONUI_API FUIInputAction
 {
 	GENERATED_BODY()
 		
@@ -107,8 +107,8 @@ public:
 	virtual void PostInitProperties() override;
 
 	int32 GetUIActionProcessingPriority() const { return UIActionProcessingPriority; }
-	const FUIInputAction* FindAction(FUIActionTag ActionTag) const;
-	const TArray<FUIInputAction>& GetUIInputActions() const { return InputActions; }
+	COMMONUI_API const FUIInputAction* FindAction(FUIActionTag ActionTag) const;
+	COMMONUI_API const TArray<FUIInputAction>& GetUIInputActions() const { return InputActions; }
 	const FCommonAnalogCursorSettings& GetAnalogCursorSettings() const { return AnalogCursorSettings; }
 //
 private:
