Fix: CommonButtonBase does not broadcast state changes

diff --git a/Engine/Plugins/Runtime/CommonUI/Source/CommonUI/Private/CommonButtonBase.cpp b/Engine/Plugins/Runtime/CommonUI/Source/CommonUI/Private/CommonButtonBase.cpp
--- a/Engine/Plugins/Runtime/CommonUI/Source/CommonUI/Private/CommonButtonBase.cpp
+++ b/Engine/Plugins/Runtime/CommonUI/Source/CommonUI/Private/CommonButtonBase.cpp
@@ -1452,6 +1452,8 @@ void UCommonButtonBase::NativeOnHovered()
 
 	NativeOnCurrentTextStyleChanged();
 	UpdateInputActionWidget();
+
+	BroadcastBinaryPostStateChange(UWidgetHoveredStateRegistration::Bit, true);
 }
 
 void UCommonButtonBase::NativeOnUnhovered()
@@ -1468,6 +1470,8 @@ void UCommonButtonBase::NativeOnUnhovered()
 
 	NativeOnCurrentTextStyleChanged();
 	UpdateInputActionWidget();
+
+	BroadcastBinaryPostStateChange(UWidgetHoveredStateRegistration::Bit, false);
 }
 
 void UCommonButtonBase::NativeOnClicked()
@@ -1529,12 +1533,14 @@ void UCommonButtonBase::NativeOnPressed()
 	HoldReset();
 	BP_OnPressed();
 	OnPressed().Broadcast();
+	BroadcastBinaryPostStateChange(UWidgetPressedStateRegistration::Bit, true);
 }
 
 void UCommonButtonBase::NativeOnReleased()
 {
 	BP_OnReleased();
 	OnReleased().Broadcast();
+	BroadcastBinaryPostStateChange(UWidgetPressedStateRegistration::Bit, false);
 }
 
 void UCommonButtonBase::NativeOnEnabled()
