Remove ensures triggered by new holdable functionality & multiplayer

diff --git a/Engine/Plugins/Runtime/CommonUI/Source/CommonUI/Private/CommonActionWidget.cpp b/Engine/Plugins/Runtime/CommonUI/Source/CommonUI/Private/CommonActionWidget.cpp
--- a/Engine/Plugins/Runtime/CommonUI/Source/CommonUI/Private/CommonActionWidget.cpp
+++ b/Engine/Plugins/Runtime/CommonUI/Source/CommonUI/Private/CommonActionWidget.cpp
@@ -344,7 +344,7 @@ const FText UCommonActionWidget::GetPaletteCategory()
 
 void UCommonActionWidget::OnActionProgress(float HeldPercent)
 {
-	if (ProgressDynamicMaterial && ensure(!ProgressMaterialParam.IsNone()))
+	if (ProgressDynamicMaterial && !ProgressMaterialParam.IsNone())
 	{
 		ProgressDynamicMaterial->SetScalarParameterValue(ProgressMaterialParam, HeldPercent);
 	}
@@ -352,7 +352,7 @@ void UCommonActionWidget::OnActionProgress(float HeldPercent)
 
 void UCommonActionWidget::OnActionComplete()
 {
-	if (ProgressDynamicMaterial && ensure(!ProgressMaterialParam.IsNone()))
+	if (ProgressDynamicMaterial && !ProgressMaterialParam.IsNone())
 	{
 		ProgressDynamicMaterial->SetScalarParameterValue(ProgressMaterialParam, 0.f);
 	}
diff --git a/Engine/Plugins/Runtime/CommonUI/Source/CommonUI/Private/Input/CommonAnalogCursor.cpp b/Engine/Plugins/Runtime/CommonUI/Source/CommonUI/Private/Input/CommonAnalogCursor.cpp
--- a/Engine/Plugins/Runtime/CommonUI/Source/CommonUI/Private/Input/CommonAnalogCursor.cpp
+++ b/Engine/Plugins/Runtime/CommonUI/Source/CommonUI/Private/Input/CommonAnalogCursor.cpp
@@ -311,7 +311,7 @@ bool FCommonAnalogCursor::HandleMouseMoveEvent(FSlateApplication& SlateApp, cons
 	{
 		SlateApp.SetPlatformCursorVisibility(true);
 		TSharedPtr<FSlateUser> SlateUser = FSlateApplication::Get().GetUser(GetOwnerUserIndex());
-		if (ensure(SlateUser))
+		if (SlateUser)
 		{
 			SlateUser->SetCursorVisibility(true);
 		}
@@ -447,7 +447,7 @@ bool FCommonAnalogCursor::IsGameViewportInFocusPathWithoutCapture() const
 		if (TSharedPtr<SViewport> GameViewportWidget = ViewportClient->GetGameViewportWidget())
 		{
 			TSharedPtr<FSlateUser> SlateUser = FSlateApplication::Get().GetUser(GetOwnerUserIndex());
-			if (ensure(SlateUser) && !SlateUser->DoesWidgetHaveCursorCapture(GameViewportWidget))
+			if (SlateUser && !SlateUser->DoesWidgetHaveCursorCapture(GameViewportWidget))
 			{
 #if PLATFORM_DESKTOP
 				// Not captured - is it in the focus path?
@@ -532,7 +532,7 @@ void FCommonAnalogCursor::HideCursor()
 void FCommonAnalogCursor::SetNormalizedCursorPosition(const FVector2D& RelativeNewPosition)
 {
 	TSharedPtr<FSlateUser> SlateUser = FSlateApplication::Get().GetUser(GetOwnerUserIndex());
-	if (ensure(SlateUser))
+	if (SlateUser)
 	{
 		const UGameViewportClient* ViewportClient = GetViewportClient();
 		if (TSharedPtr<SViewport> ViewportWidget = ViewportClient ? ViewportClient->GetGameViewportWidget() : nullptr)
