Fixed hold interaction for mobile.

#localization none
#tests Tested on mobile device.

diff --git a/Engine/Plugins/Runtime/CommonUI/Source/CommonUI/Private/CommonButtonBase.cpp b/Engine/Plugins/Runtime/CommonUI/Source/CommonUI/Private/CommonButtonBase.cpp
--- a/Engine/Plugins/Runtime/CommonUI/Source/CommonUI/Private/CommonButtonBase.cpp
+++ b/Engine/Plugins/Runtime/CommonUI/Source/CommonUI/Private/CommonButtonBase.cpp
@@ -1296,7 +1296,7 @@ void UCommonButtonBase::HandleButtonClicked()
 	{
 		// @TODO: Current click rejection method relies on click hold time, this can be refined. See NativeOnHoldProgress.
 		// Also gamepad can indirectly trigger this method, so don't guard against pressed
-    	if (bRequiresHold && CurrentHoldTime < HoldTime && IsHovered())
+    	if (bRequiresHold && CurrentHoldTime < HoldTime)
     	{
     		return;
     	}
