[26.00] Fix CommonUI actions not updating.

FActivatableTreeNode::HandleWidgetActivated:
- If the path doesn't support activation focus, manually broadcast UCommonUIActionRouterBase::OnBoundActionsUpdated delegate.

UCommonUIActionRouterBase::HandleRootNodeActivated:
- Still broadcast OnBoundActionsUpdated delegate even if activated root doesn't support activation focus.

diff --git a/Engine/Plugins/Runtime/CommonUI/Source/CommonUI/Private/Input/CommonUIActionRouterBase.cpp b/Engine/Plugins/Runtime/CommonUI/Source/CommonUI/Private/Input/CommonUIActionRouterBase.cpp
--- a/Engine/Plugins/Runtime/CommonUI/Source/CommonUI/Private/Input/CommonUIActionRouterBase.cpp
+++ b/Engine/Plugins/Runtime/CommonUI/Source/CommonUI/Private/Input/CommonUIActionRouterBase.cpp
@@ -689,12 +689,16 @@ void UCommonUIActionRouterBase::HandleRootNodeActivated(TWeakPtr<FActivatableTre
 					{
 						for (FActivatableTreeRootRef& RootNode : SortedRootList->RootList)
 						{
-							if (RootNode->IsReceivingInput() && RootNode->DoesWidgetSupportActivationFocus())
+							if (RootNode->IsReceivingInput())
 							{
 								if (RootNode == ActivatedRoot && !ActiveRootNode.IsValid())
 								{
-									// This will allow us to update the leaf node config on the activated root.
-									ActivatedRoot->UpdateLeafNode();
+									if (RootNode->DoesWidgetSupportActivationFocus())
+									{
+										// This will allow us to update the leaf node config on the activated root.
+										ActivatedRoot->UpdateLeafNode();
+									}
+
 									OnBoundActionsUpdated().Broadcast();
 								}
 								break;
diff --git a/Engine/Plugins/Runtime/CommonUI/Source/CommonUI/Private/Input/UIActionRouterTypes.cpp b/Engine/Plugins/Runtime/CommonUI/Source/CommonUI/Private/Input/UIActionRouterTypes.cpp
--- a/Engine/Plugins/Runtime/CommonUI/Source/CommonUI/Private/Input/UIActionRouterTypes.cpp
+++ b/Engine/Plugins/Runtime/CommonUI/Source/CommonUI/Private/Input/UIActionRouterTypes.cpp
@@ -1182,6 +1182,10 @@ void FActivatableTreeNode::HandleWidgetActivated()
 		{
 			GetRoot()->UpdateLeafmostActiveNode(SharedThis(this));
 		}
+		else if (RepresentedWidget.IsValid())
+		{
+			UCommonUIActionRouterBase::Get(*RepresentedWidget.Get())->OnBoundActionsUpdated().Broadcast();
+		}
 	}
 }
 
