Edigrate 25721865. Fix Action Domains not resetting.

diff --git a/Engine/Plugins/Runtime/CommonUI/Source/CommonInput/Private/CommonInputSettings.cpp b/Engine/Plugins/Runtime/CommonUI/Source/CommonInput/Private/CommonInputSettings.cpp
--- a/Engine/Plugins/Runtime/CommonUI/Source/CommonInput/Private/CommonInputSettings.cpp
+++ b/Engine/Plugins/Runtime/CommonUI/Source/CommonInput/Private/CommonInputSettings.cpp
@@ -29,6 +29,7 @@ void UCommonInputSettings::LoadData()
 void UCommonInputSettings::PostEditChangeProperty(struct FPropertyChangedEvent& PropertyChangedEvent)
 {
 	bInputDataLoaded = false;
+	bActionDomainTableLoaded = false;
 	LoadData();
 }
 #endif
