Fixed compile issue constantly breaking for most folks in usage of the common tile view.
#rnx
#tests: Compiled files to determine if the compile errors went away.

diff --git a/Engine/Plugins/Runtime/CommonUI/Source/CommonUI/Public/SCommonTileView.h b/Engine/Plugins/Runtime/CommonUI/Source/CommonUI/Public/SCommonTileView.h
--- a/Engine/Plugins/Runtime/CommonUI/Source/CommonUI/Public/SCommonTileView.h
+++ b/Engine/Plugins/Runtime/CommonUI/Source/CommonUI/Public/SCommonTileView.h
@@ -5,7 +5,7 @@
 #include "Widgets/Views/STileView.h"
 
 template <typename ItemType>
-class COMMONUI_API SCommonTileView : public STileView<ItemType>
+class SCommonTileView : public STileView<ItemType>
 {
 public:
 	virtual FReply OnFocusReceived(const FGeometry& MyGeometry, const FFocusEvent& InFocusEvent) override
