CommonUI: Add Get/Set to VisualAttachment Content Anchor property.

diff --git a/Engine/Plugins/Runtime/CommonUI/Source/CommonUI/Private/CommonVisualAttachment.cpp b/Engine/Plugins/Runtime/CommonUI/Source/CommonUI/Private/CommonVisualAttachment.cpp
--- a/Engine/Plugins/Runtime/CommonUI/Source/CommonUI/Private/CommonVisualAttachment.cpp
+++ b/Engine/Plugins/Runtime/CommonUI/Source/CommonUI/Private/CommonVisualAttachment.cpp
@@ -11,6 +11,25 @@ UCommonVisualAttachment::UCommonVisualAttachment(const FObjectInitializer& Objec
 {
 }
 
+
+FVector2D UCommonVisualAttachment::GetContentAnchor() const
+{
+	PRAGMA_DISABLE_DEPRECATION_WARNINGS
+	return ContentAnchor;
+	PRAGMA_ENABLE_DEPRECATION_WARNINGS
+}
+
+void UCommonVisualAttachment::SetContentAnchor(FVector2D InContentAnchor)
+{
+	PRAGMA_DISABLE_DEPRECATION_WARNINGS
+	ContentAnchor = InContentAnchor;
+	PRAGMA_ENABLE_DEPRECATION_WARNINGS
+	if (MyAttachmentBox)
+	{
+		MyAttachmentBox->SetContentAnchor(GetContentAnchor());
+	}
+}
+
 void UCommonVisualAttachment::ReleaseSlateResources(bool bReleaseChildren)
 {
 	Super::ReleaseSlateResources(bReleaseChildren);
@@ -35,5 +54,8 @@ void UCommonVisualAttachment::SynchronizeProperties()
 {
 	Super::SynchronizeProperties();
 
-	MyAttachmentBox->SetContentAnchor(ContentAnchor);
-}
+	if (MyAttachmentBox)
+	{
+		MyAttachmentBox->SetContentAnchor(GetContentAnchor());
+	}
+}
\ No newline at end of file
diff --git a/Engine/Plugins/Runtime/CommonUI/Source/CommonUI/Private/SVisualAttachementBox.cpp b/Engine/Plugins/Runtime/CommonUI/Source/CommonUI/Private/SVisualAttachementBox.cpp
--- a/Engine/Plugins/Runtime/CommonUI/Source/CommonUI/Private/SVisualAttachementBox.cpp
+++ b/Engine/Plugins/Runtime/CommonUI/Source/CommonUI/Private/SVisualAttachementBox.cpp
@@ -9,9 +9,11 @@ void SVisualAttachmentBox::Construct(const FArguments& InArgs)
 
 void SVisualAttachmentBox::SetContentAnchor(FVector2D InContentAnchor)
 {
-	ContentAnchor = InContentAnchor;
-
-	Invalidate(EInvalidateWidget::Layout);
+	if (InContentAnchor != ContentAnchor)
+	{
+		ContentAnchor = InContentAnchor;
+		Invalidate(EInvalidateWidget::Layout);
+	}
 }
 
 void SVisualAttachmentBox::OnArrangeChildren(const FGeometry& AllottedGeometry, FArrangedChildren& ArrangedChildren) const
diff --git a/Engine/Plugins/Runtime/CommonUI/Source/CommonUI/Public/CommonVisualAttachment.h b/Engine/Plugins/Runtime/CommonUI/Source/CommonUI/Public/CommonVisualAttachment.h
--- a/Engine/Plugins/Runtime/CommonUI/Source/CommonUI/Public/CommonVisualAttachment.h
+++ b/Engine/Plugins/Runtime/CommonUI/Source/CommonUI/Public/CommonVisualAttachment.h
@@ -17,10 +17,17 @@ class COMMONUI_API UCommonVisualAttachment : public USizeBox
 public:
 	UCommonVisualAttachment(const FObjectInitializer& ObjectInitializer);
 
-	/** */
-	UPROPERTY(EditAnywhere, BlueprintReadOnly, Category="Child Layout")
+	UE_DEPRECATED(5.4, "Direct access to ContentAnchor is deprecated. Please use the getter or setter.")
+	/** Content Anchor Point as a ratio of the content size. Use (1.0, 1.0) to anchor the content on the bottom right, (0.0) to anchor top left, (0.5, 0.5) to anchor centered.*/
+	UPROPERTY(EditAnywhere, BlueprintReadWrite, Getter, Setter, Category="Child Layout")
 	FVector2D ContentAnchor;
 
+	/** Get Content Anchor Point*/
+	FVector2D GetContentAnchor() const;
+
+	/** Set Content Anchor Point*/
+	void SetContentAnchor(FVector2D InContentAnchor);
+
 protected:
 	// UVisual interface
 	virtual void ReleaseSlateResources(bool bReleaseChildren) override;
