PR #10726: Fix CommonUI's GetDesiredFocusTarget ignoring newly added DesiredFocusWidget property.

diff --git a/Engine/Plugins/Runtime/CommonUI/Source/CommonUI/Private/CommonActivatableWidget.cpp b/Engine/Plugins/Runtime/CommonUI/Source/CommonUI/Private/CommonActivatableWidget.cpp
--- a/Engine/Plugins/Runtime/CommonUI/Source/CommonUI/Private/CommonActivatableWidget.cpp
+++ b/Engine/Plugins/Runtime/CommonUI/Source/CommonUI/Private/CommonActivatableWidget.cpp
@@ -84,12 +84,16 @@ TOptional<FActivationMetadata> UCommonActivatableWidget::GetActivationMetadata()
 
 UWidget* UCommonActivatableWidget::NativeGetDesiredFocusTarget() const
 {
-	if (bWarnDesiredFocusNotImplemented && !GetClass()->IsFunctionImplementedInScript(GET_FUNCTION_NAME_CHECKED(UCommonActivatableWidget, BP_GetDesiredFocusTarget)))
+	// Prioritize BP implementation of this function.
+	UWidget* DesiredFocusTarget = BP_GetDesiredFocusTarget();
+
+	if (!DesiredFocusTarget)
 	{
-		UE_LOG(LogCommonUI, Warning, TEXT("[%s] -> Desired Focus not implemented!"), *GetName());
+		// BP didn't specify focus target, fallback to DesiredFocusWidget property on UserWidget.
+		DesiredFocusTarget = GetDesiredFocusWidget();
 	}
 
-	return BP_GetDesiredFocusTarget();
+	return DesiredFocusTarget;
 }
 
 TOptional<FUIInputConfig> UCommonActivatableWidget::GetDesiredInputConfig() const
