PR #10744: Fix sorting of widgets in the same CommonUI domain

diff --git a/Engine/Plugins/Runtime/CommonUI/Source/CommonUI/Private/Input/CommonUIActionRouterBase.cpp b/Engine/Plugins/Runtime/CommonUI/Source/CommonUI/Private/Input/CommonUIActionRouterBase.cpp
--- a/Engine/Plugins/Runtime/CommonUI/Source/CommonUI/Private/Input/CommonUIActionRouterBase.cpp
+++ b/Engine/Plugins/Runtime/CommonUI/Source/CommonUI/Private/Input/CommonUIActionRouterBase.cpp
@@ -1689,7 +1689,7 @@ void UCommonUIActionRouterBase::FActionDomainSortedRootList::Add(FActivatableTre
 {
 	auto SortFunc = [](const FActivatableTreeRootRef& A, const FActivatableTreeRootRef& B) 
 		{
-			return A->GetLastPaintLayer() < B->GetLastPaintLayer();
+			return A->GetLastPaintLayer() > B->GetLastPaintLayer();
 		};
 
 	RootList.Add(RootNode);
