Add an InitalizeDependency to Enhanced input on the common input subsystem. Since we bind to a delegate on the EI subsystem on Initalize, we need to specify it as a dependency or else it may not be available in a packaged game.

diff --git a/Engine/Plugins/Runtime/CommonUI/Source/CommonInput/Private/CommonInputSubsystem.cpp b/Engine/Plugins/Runtime/CommonUI/Source/CommonInput/Private/CommonInputSubsystem.cpp
--- a/Engine/Plugins/Runtime/CommonUI/Source/CommonInput/Private/CommonInputSubsystem.cpp
+++ b/Engine/Plugins/Runtime/CommonUI/Source/CommonInput/Private/CommonInputSubsystem.cpp
@@ -68,6 +68,10 @@ void UCommonInputSubsystem::Initialize(FSubsystemCollectionBase& Collection)
 {
 	Super::Initialize(Collection);
 
+	// There is a dependency on the Enhanced Input subsystem below so we need to make sure it is available
+	// in a packaged game
+	Collection.InitializeDependency<UEnhancedInputLocalPlayerSubsystem>();
+
 	FCommonInputBase::GetInputSettings()->LoadData();
 
 	const UCommonInputPlatformSettings* Settings = UPlatformSettingsManager::Get().GetSettingsForPlatform<UCommonInputPlatformSettings>();
