Fix an issue where CommonButtonGroup is not selecting button when SelectedButtonIndex stays the same while the button was removed

diff --git a/Engine/Plugins/Runtime/CommonUI/Source/CommonUI/Private/Groups/CommonButtonGroupBase.cpp b/Engine/Plugins/Runtime/CommonUI/Source/CommonUI/Private/Groups/CommonButtonGroupBase.cpp
--- a/Engine/Plugins/Runtime/CommonUI/Source/CommonUI/Private/Groups/CommonButtonGroupBase.cpp
+++ b/Engine/Plugins/Runtime/CommonUI/Source/CommonUI/Private/Groups/CommonButtonGroupBase.cpp
@@ -239,17 +239,19 @@ void UCommonButtonGroupBase::OnWidgetRemoved( UWidget* OldWidget )
 		{
 			if (bSelectionRequired && Buttons.Num() > 0)
 			{
+				SelectedButtonIndex = INDEX_NONE;
+				
 				for (int32 NewButtonIndex = 0; NewButtonIndex < Buttons.Num(); NewButtonIndex++)
 				{
 					UCommonButtonBase* NewButtonToSelect = GetButtonBaseAtIndex(NewButtonIndex);
 					if (NewButtonToSelect && NewButtonToSelect->bSelectable)
 					{
+						SelectedButtonIndex = INDEX_NONE;
 						SelectButtonAtIndex(NewButtonIndex);
 						return; // Early out to only select one button
 					}
 				}
-
-				SelectedButtonIndex = INDEX_NONE;
+				
 				ensureMsgf(false, TEXT("Button group requires selection, but no button is selectable"));
 			}
 			else
