Fixed CommonTabListWidgetBase tab registration on design time

diff --git a/Engine/Plugins/Runtime/CommonUI/Source/CommonUI/Private/CommonTabListWidgetBase.cpp b/Engine/Plugins/Runtime/CommonUI/Source/CommonUI/Private/CommonTabListWidgetBase.cpp
--- a/Engine/Plugins/Runtime/CommonUI/Source/CommonUI/Private/CommonTabListWidgetBase.cpp
+++ b/Engine/Plugins/Runtime/CommonUI/Source/CommonUI/Private/CommonTabListWidgetBase.cpp
@@ -66,7 +66,8 @@ bool UCommonTabListWidgetBase::RegisterTab(FName TabNameID, TSubclassOf<UCommonB
 		return false;
 	}
 
-	UCommonButtonBase* const NewTabButton = CreateWidget<UCommonButtonBase>(GetOwningPlayer(), ButtonWidgetType);
+	// There is no PlayerController in Designer
+	UCommonButtonBase* const NewTabButton = IsDesignTime() ? CreateWidget<UCommonButtonBase>(GetWorld(), ButtonWidgetType) : CreateWidget<UCommonButtonBase>(GetOwningPlayer(), ButtonWidgetType);
 	if (!ensureMsgf(NewTabButton, TEXT("Failed to create tab button. Aborting tab registration.")))
 	{
 		return false;
