[26.30] Fix CommonUI actions not updating when widget deactivated.

FActivatableTreeNode::HandleWidgetDeactivated:
- If the path doesn't support activation focus, manually broadcast UCommonUIActionRouterBase::OnBoundActionsUpdated delegate.

diff --git a/Engine/Plugins/Runtime/CommonUI/Source/CommonUI/Private/Input/UIActionRouterTypes.cpp b/Engine/Plugins/Runtime/CommonUI/Source/CommonUI/Private/Input/UIActionRouterTypes.cpp
--- a/Engine/Plugins/Runtime/CommonUI/Source/CommonUI/Private/Input/UIActionRouterTypes.cpp
+++ b/Engine/Plugins/Runtime/CommonUI/Source/CommonUI/Private/Input/UIActionRouterTypes.cpp
@@ -1196,18 +1196,25 @@ void FActivatableTreeNode::HandleWidgetDeactivated()
 	if (bCanReceiveInput)
 	{
 		SetCanReceiveInput(false);
-		
-		if (Parent.IsValid() && DoesPathSupportActivationFocus())
+
+		if (Parent.IsValid())
 		{
-			// Search for the nearest parent that's still receiving input to give it focus
-			if (FActivatableTreeNodePtr NearestActiveParent = GetParentNode())
+			if (DoesPathSupportActivationFocus())
 			{
-				while (NearestActiveParent && !NearestActiveParent->IsReceivingInput())
+				// Search for the nearest parent that's still receiving input to give it focus
+				if (FActivatableTreeNodePtr NearestActiveParent = GetParentNode())
 				{
-					NearestActiveParent = NearestActiveParent->GetParentNode();
-				}
+					while (NearestActiveParent && !NearestActiveParent->IsReceivingInput())
+					{
+						NearestActiveParent = NearestActiveParent->GetParentNode();
+					}
 
-				GetActionRouter().UpdateLeafNodeAndConfig(GetRoot(), NearestActiveParent);
+					GetActionRouter().UpdateLeafNodeAndConfig(GetRoot(), NearestActiveParent);
+				}
+			}
+			else
+			{
+				GetActionRouter().OnBoundActionsUpdated().Broadcast();
 			}
 		}
 	}
