Common UI:
- Add a design time FKey variable in CommonActionWidget for preview

diff --git a/Engine/Plugins/Runtime/CommonUI/Source/CommonUI/Private/CommonActionWidget.cpp b/Engine/Plugins/Runtime/CommonUI/Source/CommonUI/Private/CommonActionWidget.cpp
--- a/Engine/Plugins/Runtime/CommonUI/Source/CommonUI/Private/CommonActionWidget.cpp
+++ b/Engine/Plugins/Runtime/CommonUI/Source/CommonUI/Private/CommonActionWidget.cpp
@@ -2,6 +2,7 @@
 
 #include "CommonActionWidget.h"
 
+#include "CommonInputBaseTypes.h"
 #include "CommonInputSubsystem.h"
 #include "CommonInputTypeEnum.h"
 #include "CommonUITypes.h"
@@ -112,12 +113,42 @@ void UCommonActionWidget::SynchronizeProperties()
 
 FSlateBrush UCommonActionWidget::GetIcon() const
 {
-	if (const UCommonInputSubsystem* CommonInputSubsystem = GetInputSubsystem())
+	if (!IsDesignTime())
 	{
-		return EnhancedInputAction && CommonUI::IsEnhancedInputSupportEnabled()
-			? CommonUI::GetIconForEnhancedInputAction(CommonInputSubsystem, EnhancedInputAction)
-			: CommonUI::GetIconForInputActions(CommonInputSubsystem, InputActions);
+		if (const UCommonInputSubsystem* CommonInputSubsystem = GetInputSubsystem())
+		{
+			return EnhancedInputAction && CommonUI::IsEnhancedInputSupportEnabled()
+				? CommonUI::GetIconForEnhancedInputAction(CommonInputSubsystem, EnhancedInputAction)
+				: CommonUI::GetIconForInputActions(CommonInputSubsystem, InputActions);
+		}
 	}
+#if WITH_EDITORONLY_DATA
+	else
+	{
+		if (DesignTimeKey.IsValid())
+		{
+			ECommonInputType Dummy;
+			FName OutDefaultGamepadName;
+			FCommonInputBase::GetCurrentPlatformDefaults(Dummy, OutDefaultGamepadName);
+
+			ECommonInputType KeyInputType = ECommonInputType::MouseAndKeyboard;
+			if (DesignTimeKey.IsGamepadKey())
+			{
+				KeyInputType = ECommonInputType::Gamepad;
+			}
+			else if (DesignTimeKey.IsTouch())
+			{
+				KeyInputType = ECommonInputType::Touch;
+			}
+
+			FSlateBrush InputBrush;
+			if (UCommonInputPlatformSettings::Get()->TryGetInputBrush(InputBrush, TArray<FKey> { DesignTimeKey }, KeyInputType, OutDefaultGamepadName))
+			{
+				return InputBrush;
+			}
+		}
+	}
+#endif
 
 	return *FStyleDefaults::GetNoBrush();
 }
@@ -264,13 +295,16 @@ void UCommonActionWidget::UpdateBindingHandleInternal(FUIActionBindingHandle Bin
 
 void UCommonActionWidget::UpdateActionWidget()
 {
-	if (!IsDesignTime() && GetWorld())
+	if (GetWorld())
 	{
 		const UCommonInputSubsystem* CommonInputSubsystem = GetInputSubsystem();
-		if (GetGameInstance() && ensure(CommonInputSubsystem) && CommonInputSubsystem->ShouldShowInputKeys())
+		if (IsDesignTime() || (GetGameInstance() && ensure(CommonInputSubsystem) && CommonInputSubsystem->ShouldShowInputKeys()))
 		{
 			const FCommonInputActionDataBase* InputActionData = GetInputActionData();
-			if (InputActionData || (EnhancedInputAction && CommonUI::IsEnhancedInputSupportEnabled()))
+			bool bIsEnhancedInputAction = EnhancedInputAction && CommonUI::IsEnhancedInputSupportEnabled();
+			bool bIsDesignPreview = IsDesignTime() && DesignTimeKey.IsValid();
+
+			if (InputActionData || bIsEnhancedInputAction || bIsDesignPreview)
 			{
 				if (bAlwaysHideOverride)
 				{
diff --git a/Engine/Plugins/Runtime/CommonUI/Source/CommonUI/Public/CommonActionWidget.h b/Engine/Plugins/Runtime/CommonUI/Source/CommonUI/Public/CommonActionWidget.h
--- a/Engine/Plugins/Runtime/CommonUI/Source/CommonUI/Public/CommonActionWidget.h
+++ b/Engine/Plugins/Runtime/CommonUI/Source/CommonUI/Public/CommonActionWidget.h
@@ -106,6 +106,9 @@ protected:
 #if WITH_EDITORONLY_DATA
 	UPROPERTY()
 	FDataTableRowHandle InputActionDataRow_DEPRECATED;
+
+	UPROPERTY(EditAnywhere, Category = "CommonActionWidget")
+	FKey DesignTimeKey;
 #endif
 
 protected:
