CIS Fix non-editor

diff --git a/Engine/Plugins/Runtime/CommonUI/Source/CommonUI/Private/CommonActionWidget.cpp b/Engine/Plugins/Runtime/CommonUI/Source/CommonUI/Private/CommonActionWidget.cpp
--- a/Engine/Plugins/Runtime/CommonUI/Source/CommonUI/Private/CommonActionWidget.cpp
+++ b/Engine/Plugins/Runtime/CommonUI/Source/CommonUI/Private/CommonActionWidget.cpp
@@ -301,9 +301,13 @@ void UCommonActionWidget::UpdateActionWidget()
 		if (IsDesignTime() || (GetGameInstance() && ensure(CommonInputSubsystem) && CommonInputSubsystem->ShouldShowInputKeys()))
 		{
 			const FCommonInputActionDataBase* InputActionData = GetInputActionData();
-			bool bIsEnhancedInputAction = EnhancedInputAction && CommonUI::IsEnhancedInputSupportEnabled();
-			bool bIsDesignPreview = IsDesignTime() && DesignTimeKey.IsValid();
+			const bool bIsEnhancedInputAction = EnhancedInputAction && CommonUI::IsEnhancedInputSupportEnabled();
 
+#if WITH_EDITORONLY_DATA
+			const bool bIsDesignPreview = IsDesignTime() && DesignTimeKey.IsValid();
+#else
+			const bool bIsDesignPreview = false;
+#endif
 			if (InputActionData || bIsEnhancedInputAction || bIsDesignPreview)
 			{
 				if (bAlwaysHideOverride)
