[26.30]

FUIActionBindingHandle:
- Added GetDisplayInActionBar and SetDisplayInActionBar functions.

diff --git a/Engine/Plugins/Runtime/CommonUI/Source/CommonUI/Private/Input/UIActionRouterTypes.cpp b/Engine/Plugins/Runtime/CommonUI/Source/CommonUI/Private/Input/UIActionRouterTypes.cpp
--- a/Engine/Plugins/Runtime/CommonUI/Source/CommonUI/Private/Input/UIActionRouterTypes.cpp
+++ b/Engine/Plugins/Runtime/CommonUI/Source/CommonUI/Private/Input/UIActionRouterTypes.cpp
@@ -530,6 +530,36 @@ void FUIActionBindingHandle::SetDisplayName(const FText& DisplayName)
 	}
 }
 
+bool FUIActionBindingHandle::GetDisplayInActionBar() const
+{
+	bool bDisplayInActionBar = false;
+
+	if (const FUIActionBinding* Binding = FUIActionBinding::FindBinding(*this).Get())
+	{
+		bDisplayInActionBar = Binding->bDisplayInActionBar;
+	}
+
+	return bDisplayInActionBar;
+}
+
+void FUIActionBindingHandle::SetDisplayInActionBar(const bool bDisplayInActionBar)
+{
+	FUIActionBinding* Binding = FUIActionBinding::FindBinding(*this).Get();
+
+	if (Binding && Binding->bDisplayInActionBar != bDisplayInActionBar)
+	{
+		Binding->bDisplayInActionBar = bDisplayInActionBar;
+
+		if (const UWidget* BoundWidget = Binding->BoundWidget.Get())
+		{
+			if (const UCommonUIActionRouterBase* ActionRouter = UCommonUIActionRouterBase::Get(*BoundWidget))
+			{
+				ActionRouter->OnBoundActionsUpdated().Broadcast();
+			}
+		}
+	}
+}
+
 const UWidget* FUIActionBindingHandle::GetBoundWidget() const
 {
 	if (TSharedPtr<const FUIActionBinding> Binding = FUIActionBinding::FindBinding(*this))
diff --git a/Engine/Plugins/Runtime/CommonUI/Source/CommonUI/Public/Input/UIActionBindingHandle.h b/Engine/Plugins/Runtime/CommonUI/Source/CommonUI/Public/Input/UIActionBindingHandle.h
--- a/Engine/Plugins/Runtime/CommonUI/Source/CommonUI/Public/Input/UIActionBindingHandle.h
+++ b/Engine/Plugins/Runtime/CommonUI/Source/CommonUI/Public/Input/UIActionBindingHandle.h
@@ -28,6 +28,11 @@ public:
 	/** Should not be called often as broadcasts UCommonUIActionRouterBase::OnBoundActionsUpdated event */
 	void SetDisplayName(const FText& DisplayName);
 
+	bool GetDisplayInActionBar() const;
+
+	/** Should not be called often as broadcasts UCommonUIActionRouterBase::OnBoundActionsUpdated event */
+	void SetDisplayInActionBar(const bool bDisplayInActionBar);
+
 	const UWidget* GetBoundWidget() const;
 
 	FUIActionBindingHandle() {}
