Fix for bound input action triggering hold interaction.
Added implementable event for bounded action triggered.
#rnx

diff --git a/Engine/Plugins/Runtime/CommonUI/Source/CommonUI/Private/CommonButtonBase.cpp b/Engine/Plugins/Runtime/CommonUI/Source/CommonUI/Private/CommonButtonBase.cpp
--- a/Engine/Plugins/Runtime/CommonUI/Source/CommonUI/Private/CommonButtonBase.cpp
+++ b/Engine/Plugins/Runtime/CommonUI/Source/CommonUI/Private/CommonButtonBase.cpp
@@ -356,6 +356,7 @@ UCommonButtonBase::UCommonButtonBase(const FObjectInitializer& ObjectInitializer
 	, HoldTime(0.f)
 	, HoldRollbackTime(0.f)
 	, CurrentHoldTime(0.f)
+	, CurrentHoldProgress(0.f)
 {
 	SetIsFocusable(true);
 }
@@ -633,6 +634,7 @@ void UCommonButtonBase::UnbindTriggeringInputActionToClick()
 	}
 	
 	CurrentHoldTime = 0.f;
+	CurrentHoldProgress = 0.f;
 }
 
 void UCommonButtonBase::HandleTriggeringActionCommited(bool& bPassthrough)
@@ -642,8 +644,12 @@ void UCommonButtonBase::HandleTriggeringActionCommited(bool& bPassthrough)
 
 void UCommonButtonBase::HandleTriggeringActionCommited()
 {
-	// Because this path doesn't go through SButton::Press(), the sound needs to be played from here.
-	FSlateApplication::Get().PlaySound(NormalStyle.PressedSlateSound);
+	if (IsInteractionEnabled())
+	{
+		// Because this path doesn't go through SButton::Press(), the sound needs to be played from here.
+		FSlateApplication::Get().PlaySound(NormalStyle.PressedSlateSound);
+		BP_OnInputActionTriggered();
+	}
 	HandleButtonClicked();
 }
 
@@ -795,6 +801,7 @@ void UCommonButtonBase::NativeOnActionProgress(float HeldPercent)
 		InputActionWidget->OnActionProgress(HeldPercent);
 	}
 	OnActionProgress(HeldPercent);
+	CurrentHoldProgress = HeldPercent;
 }
 
 bool UCommonButtonBase::NativeOnHoldProgress(float DeltaTime)
@@ -802,9 +809,9 @@ bool UCommonButtonBase::NativeOnHoldProgress(float DeltaTime)
 	if (HoldTime > UE_SMALL_NUMBER)
 	{
 		CurrentHoldTime += FMath::Clamp(DeltaTime, 0.f, HoldTime);
-		const float HeldPercent = FMath::Clamp(CurrentHoldTime / HoldTime, 0.f, 1.f);
-		NativeOnActionProgress(HeldPercent);
-		if (HeldPercent >= 1.f)
+		CurrentHoldProgress = FMath::Clamp(CurrentHoldTime / HoldTime, 0.f, 1.f);
+		NativeOnActionProgress(CurrentHoldProgress);
+		if (CurrentHoldProgress >= 1.f)
 		{
 			HandleTriggeringActionCommited();
 			HoldReset();
@@ -814,7 +821,6 @@ bool UCommonButtonBase::NativeOnHoldProgress(float DeltaTime)
 		return true;
 	}
 	HoldReset();
-	
 	return false;
 }
 
@@ -824,9 +830,9 @@ bool UCommonButtonBase::NativeOnHoldProgressRollback(float DeltaTime)
 	{
 		const float HoldRollbackMultiplier = HoldTime / HoldRollbackTime;
 		CurrentHoldTime = FMath::Clamp(CurrentHoldTime - (DeltaTime * HoldRollbackMultiplier), 0.f, HoldRollbackTime);
-		const float HoldRollbackPercent = FMath::Clamp(CurrentHoldTime / HoldTime, 0.f, 1.f);
-		NativeOnActionProgress(HoldRollbackPercent);
-		if (HoldRollbackPercent <= 0.f)
+		CurrentHoldProgress = FMath::Clamp(CurrentHoldTime / HoldTime, 0.f, 1.f);
+		NativeOnActionProgress(CurrentHoldProgress);
+		if (CurrentHoldProgress <= 0.f)
 		{
 			FTSTicker::GetCoreTicker().RemoveTicker(HoldProgressRollbackTickerHandle);
 			HoldProgressRollbackTickerHandle = nullptr;
@@ -853,7 +859,8 @@ void UCommonButtonBase::HoldReset()
 		FTSTicker::GetCoreTicker().RemoveTicker(HoldProgressRollbackTickerHandle);
 		HoldProgressRollbackTickerHandle = nullptr;
 	}
-	CurrentHoldTime = 0.0f;
+	CurrentHoldTime = 0.f;
+	CurrentHoldProgress = 0.f;
 }
 
 void UCommonButtonBase::NativeOnActionComplete()
@@ -1296,7 +1303,7 @@ void UCommonButtonBase::HandleButtonClicked()
 	{
 		// @TODO: Current click rejection method relies on click hold time, this can be refined. See NativeOnHoldProgress.
 		// Also gamepad can indirectly trigger this method, so don't guard against pressed
-    	if (bRequiresHold && CurrentHoldTime < HoldTime)
+    	if (bRequiresHold && CurrentHoldProgress < 1.f)
     	{
     		return;
     	}
@@ -1385,7 +1392,6 @@ void UCommonButtonBase::HandleButtonReleased()
 		if (HoldRollbackTime <= UE_SMALL_NUMBER)
 		{
 			HoldReset();
-			NativeOnHoldProgress(0.f);
 		}
 		else
 		{
diff --git a/Engine/Plugins/Runtime/CommonUI/Source/CommonUI/Public/CommonButtonBase.h b/Engine/Plugins/Runtime/CommonUI/Source/CommonUI/Public/CommonButtonBase.h
--- a/Engine/Plugins/Runtime/CommonUI/Source/CommonUI/Public/CommonButtonBase.h
+++ b/Engine/Plugins/Runtime/CommonUI/Source/CommonUI/Public/CommonButtonBase.h
@@ -592,6 +592,9 @@ protected:
 	void BP_OnDoubleClicked();
 	virtual void NativeOnDoubleClicked();
 
+	UFUNCTION(BlueprintImplementableEvent, Category = CommonButton, meta = (DisplayName = "On Input Action Triggered"))
+	void BP_OnInputActionTriggered();
+
 	/** Unless this is called, we will assume the double click should be converted into a normal click. */
 	UFUNCTION(BlueprintCallable, Category = CommonButton)
 	void StopDoubleClickPropagation();
@@ -867,6 +870,9 @@ protected:
 	
     /** Current hold time for this button */
     float CurrentHoldTime;
+
+	/** Current hold progress % for this button */
+	float CurrentHoldProgress;
 	
     /** Handle for ticker spawned for press and hold */
     FTSTicker::FDelegateHandle HoldTickerHandle;
