Fixed issue in CommonTextBlock with LineHeightPercentage and Margin are not updating from bp.
#rnx

diff --git a/Engine/Plugins/Runtime/CommonUI/Source/CommonUI/Private/CommonTextBlock.cpp b/Engine/Plugins/Runtime/CommonUI/Source/CommonUI/Private/CommonTextBlock.cpp
--- a/Engine/Plugins/Runtime/CommonUI/Source/CommonUI/Private/CommonTextBlock.cpp
+++ b/Engine/Plugins/Runtime/CommonUI/Source/CommonUI/Private/CommonTextBlock.cpp
@@ -349,6 +349,7 @@ void UCommonTextBlock::SetTextCase(bool bUseAllCaps)
 void UCommonTextBlock::SetLineHeightPercentage(float InLineHeightPercentage)
 {
 	LineHeightPercentage = InLineHeightPercentage;
+	SynchronizeProperties();
 }
 
 void UCommonTextBlock::SetStyle(TSubclassOf<UCommonTextStyle> InStyle)
@@ -365,6 +366,7 @@ const FMargin& UCommonTextBlock::GetMargin()
 void UCommonTextBlock::SetMargin(const FMargin& InMargin)
 {
 	Margin = InMargin;
+	SynchronizeProperties();
 }
 
 void UCommonTextBlock::ResetScrollState()
