PR #10824: Fix not being able to block input with multiple reasons with UCommonInputSubsystem::SetInputTypeFilter

diff --git a/Engine/Plugins/Runtime/CommonUI/Source/CommonInput/Private/CommonInputPreprocessor.cpp b/Engine/Plugins/Runtime/CommonUI/Source/CommonInput/Private/CommonInputPreprocessor.cpp
--- a/Engine/Plugins/Runtime/CommonUI/Source/CommonInput/Private/CommonInputPreprocessor.cpp
+++ b/Engine/Plugins/Runtime/CommonUI/Source/CommonInput/Private/CommonInputPreprocessor.cpp
@@ -119,7 +119,7 @@ bool FCommonInputPreprocessor::HandleMouseButtonDoubleClickEvent(FSlateApplicati
 
 void FCommonInputPreprocessor::SetInputTypeFilter(ECommonInputType InputType, FName InReason, bool InFilter)
 {
-	TMap<FName, bool> Reasons = FilterInputTypeWithReasons[(uint8)InputType];
+	TMap<FName, bool>& Reasons = FilterInputTypeWithReasons[(uint8)InputType];
 	Reasons.Add(InReason, InFilter);
 
 	bool ComputedFilter = false;
@@ -260,4 +260,4 @@ ECommonInputType FCommonInputPreprocessor::GetInputType(const FPointerEvent& Poi
 		return ECommonInputType::Gamepad;
 	}
 	return ECommonInputType::MouseAndKeyboard;
-}
\ No newline at end of file
+}
