PR #10817: Fix UCommonActivatableWidgetSwitcher activating the first tab when changing tabs first time

diff --git a/Engine/Plugins/Runtime/CommonUI/Source/CommonUI/Private/CommonAnimatedSwitcher.cpp b/Engine/Plugins/Runtime/CommonUI/Source/CommonUI/Private/CommonAnimatedSwitcher.cpp
--- a/Engine/Plugins/Runtime/CommonUI/Source/CommonUI/Private/CommonAnimatedSwitcher.cpp
+++ b/Engine/Plugins/Runtime/CommonUI/Source/CommonUI/Private/CommonAnimatedSwitcher.cpp
@@ -164,6 +164,9 @@ PRAGMA_DISABLE_DEPRECATION_WARNINGS
 	if (Index >= 0 && Index < Slots.Num() && (Index != ActiveWidgetIndex || !bSetOnce))
 	{
 		HandleOutgoingWidget();
+		
+		// When we're setting up, and the index goes 0->0, MyAnimatedSwitcher won't fire its ActiveIndexChanged Event, so we manually call HandleSlateActiveIndexChanged below.
+		const bool bIsSettingInitialIndex = !bSetOnce && (!MyAnimatedSwitcher.IsValid() || Index == MyAnimatedSwitcher->GetActiveWidgetIndex());
 
 		// For now we can't call setter since it calls MyWidgetSwitcher->SetActiveWidgetIndex(SafeIndex)
 		if (ActiveWidgetIndex != Index)
@@ -179,8 +182,7 @@ PRAGMA_DISABLE_DEPRECATION_WARNINGS
 			MyAnimatedSwitcher->TransitionToIndex(SafeIndex, bInstantTransition);
 		}
 
-		//When we're setting up, and the index goes 0->0, MyAnimatedSwitcher won't fire its ActiveIndexChanged Event.
-		if (!bSetOnce)
+		if (bIsSettingInitialIndex)
 		{
 			HandleSlateActiveIndexChanged(ActiveWidgetIndex);
 		}
