PR #10770: Fix CommonBoundActionBar Binding filter not using ignore duplicates

diff --git a/Engine/Plugins/Runtime/CommonUI/Source/CommonUI/Private/Input/CommonBoundActionBar.cpp b/Engine/Plugins/Runtime/CommonUI/Source/CommonUI/Private/Input/CommonBoundActionBar.cpp
--- a/Engine/Plugins/Runtime/CommonUI/Source/CommonUI/Private/Input/CommonBoundActionBar.cpp
+++ b/Engine/Plugins/Runtime/CommonUI/Source/CommonUI/Private/Input/CommonBoundActionBar.cpp
@@ -158,7 +158,7 @@ void UCommonBoundActionBar::HandleDeferredDisplayUpdate()
 					const FName& PlayerGamepadName = InputSubsystem.GetCurrentGamepadName();
 
 					TSet<FName> AcceptedBindings;
-					TArray<FUIActionBindingHandle> FilteredBindings = ActionRouter->GatherActiveBindings().FilterByPredicate([ActionRouter, PlayerInputType, PlayerGamepadName, &AcceptedBindings](const FUIActionBindingHandle& Handle) mutable
+					TArray<FUIActionBindingHandle> FilteredBindings = ActionRouter->GatherActiveBindings().FilterByPredicate([ActionRouter, PlayerInputType, PlayerGamepadName, &AcceptedBindings, this](const FUIActionBindingHandle& Handle) mutable
 						{
 							if (TSharedPtr<FUIActionBinding> Binding = FUIActionBinding::FindBinding(Handle))
 							{
@@ -171,7 +171,17 @@ void UCommonBoundActionBar::HandleDeferredDisplayUpdate()
 								{
 									if (TObjectPtr<const UInputAction> InputAction = Binding->InputAction.Get())
 									{
-										return CommonUI::ActionValidForInputType(ActionRouter->GetLocalPlayer(), PlayerInputType, InputAction);
+										if (CommonUI::ActionValidForInputType(ActionRouter->GetLocalPlayer(), PlayerInputType, InputAction))
+										{
+											if (!bIgnoreDuplicateActions)
+											{
+												return true;
+											}
+											bool bAlreadyAccepted = false;
+											AcceptedBindings.Add(Binding->ActionName, &bAlreadyAccepted);
+											return !bAlreadyAccepted;
+										}
+										return false;
 									}
 								}
 
@@ -187,6 +197,10 @@ void UCommonBoundActionBar::HandleDeferredDisplayUpdate()
 									return false; 
 								}
 
+								if (!bIgnoreDuplicateActions)
+								{
+									return true;
+								}
 								bool bAlreadyAccepted = false;
 								AcceptedBindings.Add(Binding->ActionName, &bAlreadyAccepted);
 								return !bAlreadyAccepted;
