PR #10596: CommonWidgetCarousel MoveSpeed exposed

diff --git a/Engine/Plugins/Runtime/CommonUI/Source/CommonUI/Private/CommonWidgetCarousel.cpp b/Engine/Plugins/Runtime/CommonUI/Source/CommonUI/Private/CommonWidgetCarousel.cpp
--- a/Engine/Plugins/Runtime/CommonUI/Source/CommonUI/Private/CommonWidgetCarousel.cpp
+++ b/Engine/Plugins/Runtime/CommonUI/Source/CommonUI/Private/CommonWidgetCarousel.cpp
@@ -10,6 +10,8 @@ UCommonWidgetCarousel::UCommonWidgetCarousel(const FObjectInitializer& ObjectIni
 {
 	bIsVariable = true;
 	SetClipping(EWidgetClipping::ClipToBounds);
+
+	MoveSpeed = 5.f;
 }
 
 void UCommonWidgetCarousel::ReleaseSlateResources(bool bReleaseChildren)
@@ -65,6 +67,28 @@ void UCommonWidgetCarousel::PreviousPage()
 	}
 }
 
+void UCommonWidgetCarousel::SetMoveSpeed(float InMoveSpeed)
+{
+	MoveSpeed = InMoveSpeed;
+	if (MoveSpeed != InMoveSpeed)
+	{
+		if (MyCommonWidgetCarousel.IsValid())
+		{
+			MyCommonWidgetCarousel->SetMoveSpeed(MoveSpeed);
+		}
+	}
+}
+
+float UCommonWidgetCarousel::GetMoveSpeed() const
+{
+	if (MyCommonWidgetCarousel.IsValid())
+	{
+		return MyCommonWidgetCarousel->GetMoveSpeed();
+	}
+
+	return MoveSpeed;
+}
+
 int32 UCommonWidgetCarousel::GetActiveWidgetIndex() const
 {
 	if ( MyCommonWidgetCarousel.IsValid() )
@@ -127,7 +151,7 @@ TSharedRef<SWidget> UCommonWidgetCarousel::RebuildWidget()
 		.FadeRate(0)
 		.SlideValueLeftLimit(-1)
 		.SlideValueRightLimit(1)
-		.MoveSpeed(5.f)
+		.MoveSpeed(MoveSpeed)
 		.OnGenerateWidget_UObject(this, &UCommonWidgetCarousel::OnGenerateWidgetForCarousel)
 		.OnPageChanged_UObject(this, &UCommonWidgetCarousel::HandlePageChanged);
 
diff --git a/Engine/Plugins/Runtime/CommonUI/Source/CommonUI/Public/CommonWidgetCarousel.h b/Engine/Plugins/Runtime/CommonUI/Source/CommonUI/Public/CommonWidgetCarousel.h
--- a/Engine/Plugins/Runtime/CommonUI/Source/CommonUI/Public/CommonWidgetCarousel.h
+++ b/Engine/Plugins/Runtime/CommonUI/Source/CommonUI/Public/CommonWidgetCarousel.h
@@ -27,6 +27,12 @@ public:
 	UPROPERTY(EditAnywhere, BlueprintReadOnly, Category="Carousel", meta=( UIMin=0, ClampMin=0 ))
 	int32 ActiveWidgetIndex;
 
+protected:
+
+	/** How quickly the carousel transitions when changing active widget */
+	UPROPERTY(EditAnywhere, BlueprintReadWrite, Getter, Setter, BlueprintGetter="GetMoveSpeed", BlueprintSetter="SetMoveSpeed", Category="Carousel", meta=( UIMin=0, ClampMin=0 ))
+	float MoveSpeed;
+	
 public:
 
 	/** Gets the slot index of the currently active widget */
@@ -56,6 +62,14 @@ public:
 
 	UFUNCTION(BlueprintCallable, Category="Carousel")
 	void PreviousPage();
+
+	/** Sets the Move Speed. */
+	UFUNCTION(BlueprintCallable, Category="Widget")
+	void SetMoveSpeed(float InMoveSpeed);
+
+	/** Gets the Move Speed. */
+	UFUNCTION(BlueprintCallable, Category = "Widget")
+	float GetMoveSpeed() const;
 	
 	UPROPERTY(BlueprintAssignable, Category = "Carousel")
 	FOnCurrentPageIndexChanged OnCurrentPageIndexChanged;
diff --git a/Engine/Source/Runtime/WidgetCarousel/Public/SWidgetCarousel.h b/Engine/Source/Runtime/WidgetCarousel/Public/SWidgetCarousel.h
--- a/Engine/Source/Runtime/WidgetCarousel/Public/SWidgetCarousel.h
+++ b/Engine/Source/Runtime/WidgetCarousel/Public/SWidgetCarousel.h
@@ -322,6 +322,11 @@ struct FCarouselDisplayItem : public TSharedFromThis<FCarouselDisplayItem>
 		return SlideValue;
 	}
 
+	const float GetMoveSpeed() const
+	{
+		return MoveSpeed;
+	}
+
 protected:
 
 	/**
@@ -602,6 +607,21 @@ public:
 		return CenterCarouselWidget->GetSlideValue();
 	}
 
+	void SetMoveSpeed(float InMoveSpeed)
+	{
+		if (!IsInTransition())
+		{
+			CenterCarouselWidget->SetMoveSpeed(InMoveSpeed);
+			RightCarouselWidget->SetMoveSpeed(InMoveSpeed);
+			LeftCarouselWidget->SetMoveSpeed(InMoveSpeed);
+		}
+	}
+
+	float GetMoveSpeed() const
+	{
+		return CenterCarouselWidget->GetMoveSpeed();
+	}
+
 protected:
 
 	/**
