[26.30]

UCommonVisibilitySwitcher:
- Added bCurrentlySwitching flag, set to true in SetActiveWidgetIndex_Internal.
- Added BlueprintPure IsCurrentlySwitching getter.

#changelist validated

#virtualized

diff --git a/Engine/Plugins/Runtime/CommonUI/Source/CommonUI/Private/CommonVisibilitySwitcher.cpp b/Engine/Plugins/Runtime/CommonUI/Source/CommonUI/Private/CommonVisibilitySwitcher.cpp
--- a/Engine/Plugins/Runtime/CommonUI/Source/CommonUI/Private/CommonVisibilitySwitcher.cpp
+++ b/Engine/Plugins/Runtime/CommonUI/Source/CommonUI/Private/CommonVisibilitySwitcher.cpp
@@ -114,6 +114,8 @@ UWidget* UCommonVisibilitySwitcher::GetWidgetAtIndex(int32 Index) const
 
 void UCommonVisibilitySwitcher::SetActiveWidgetIndex_Internal(int32 Index, bool bBroadcastChange /*= true*/)
 {
+	TGuardValue<bool> bCurrentlySwitchingGuard(bCurrentlySwitching, true);
+
 	if (Slots.IsValidIndex(ActiveWidgetIndex))
 	{
 		if (UCommonVisibilitySwitcherSlot* OldActiveSlot = Cast<UCommonVisibilitySwitcherSlot>(Slots[ActiveWidgetIndex]))
diff --git a/Engine/Plugins/Runtime/CommonUI/Source/CommonUI/Public/CommonVisibilitySwitcher.h b/Engine/Plugins/Runtime/CommonUI/Source/CommonUI/Public/CommonVisibilitySwitcher.h
--- a/Engine/Plugins/Runtime/CommonUI/Source/CommonUI/Public/CommonVisibilitySwitcher.h
+++ b/Engine/Plugins/Runtime/CommonUI/Source/CommonUI/Public/CommonVisibilitySwitcher.h
@@ -40,6 +40,9 @@ public:
 	UFUNCTION(BlueprintCallable, Category = CommonVisibilitySwitcher)
 	void DeactivateVisibleSlot();
 
+	UFUNCTION(BlueprintPure, Category = CommonVisibilitySwitcher)
+	bool IsCurrentlySwitching() const { return bCurrentlySwitching; }
+
 	UWidget* GetWidgetAtIndex(int32 Index) const;
 
 	virtual void SynchronizeProperties() override;
@@ -93,5 +96,7 @@ protected:
 	UPROPERTY(EditAnywhere, Category = CommonVisibilitySwitcher)
 	bool bActivateFirstSlotOnAdding = false;
 
+	bool bCurrentlySwitching = false;
+
 	mutable FOnActiveWidgetIndexChanged OnActiveWidgetIndexChangedEvent;
 };
