Common UI:
- Remove force OnDemand clipping when using a Text Scroll Style

diff --git a/Engine/Plugins/Runtime/CommonUI/Source/CommonUI/Private/CommonTextBlock.cpp b/Engine/Plugins/Runtime/CommonUI/Source/CommonUI/Private/CommonTextBlock.cpp
--- a/Engine/Plugins/Runtime/CommonUI/Source/CommonUI/Private/CommonTextBlock.cpp
+++ b/Engine/Plugins/Runtime/CommonUI/Source/CommonUI/Private/CommonTextBlock.cpp
@@ -505,6 +505,8 @@ bool UCommonTextBlock::CanEditChange(const FProperty* InProperty) const
 	if (Super::CanEditChange(InProperty))
 	{
 PRAGMA_DISABLE_DEPRECATION_WARNINGS
+		bool bIsInvalidProperty = false;
+		
 		if (const UCommonTextStyle* TextStyle = GetStyleCDO())
 		{
 			static TArray<FName> InvalidPropertiesWithStyle =
@@ -518,8 +520,24 @@ PRAGMA_DISABLE_DEPRECATION_WARNINGS
 				GET_MEMBER_NAME_CHECKED(UCommonTextBlock, ShadowColorAndOpacity)
 			};
 
-			return !InvalidPropertiesWithStyle.Contains(InProperty->GetFName());
+			bIsInvalidProperty |= InvalidPropertiesWithStyle.Contains(InProperty->GetFName());
 		}
+		
+		if (const UCommonTextScrollStyle* TextScrollStyle = UCommonTextBlock::GetScrollStyleCDO())
+		{
+			static TArray<FName> InvalidPropertiesWithScrollStyle =
+			{
+				GET_MEMBER_NAME_CHECKED(UCommonTextBlock, Clipping)
+			};
+			
+			bIsInvalidProperty |= InvalidPropertiesWithScrollStyle.Contains(InProperty->GetFName());
+		}
+
+		if (bIsInvalidProperty)
+		{
+			return false;
+		}
+
 		if (bAutoCollapseWithEmptyText && InProperty->GetFName() == GET_MEMBER_NAME_CHECKED(UCommonTextBlock, Visibility))
 		{
 			return false;
@@ -541,12 +559,7 @@ TSharedRef<SWidget> UCommonTextBlock::RebuildWidget()
 		return Super::RebuildWidget();
 	}
 
-	// If the clipping mode is the default, but we're using a scrolling style,
-	// we need to switch over to a clip to bounds style.
-	if (GetClipping() == EWidgetClipping::Inherit)
-	{
-		SetClipping(EWidgetClipping::OnDemand);
-	}
+	SetClipping(TextScrollStyle->Clipping);
 
 	// clang-format off
 	TextScroller = 
diff --git a/Engine/Plugins/Runtime/CommonUI/Source/CommonUI/Public/CommonTextBlock.h b/Engine/Plugins/Runtime/CommonUI/Source/CommonUI/Public/CommonTextBlock.h
--- a/Engine/Plugins/Runtime/CommonUI/Source/CommonUI/Public/CommonTextBlock.h
+++ b/Engine/Plugins/Runtime/CommonUI/Source/CommonUI/Public/CommonTextBlock.h
@@ -108,6 +108,9 @@ public:
 
 	UPROPERTY(EditAnywhere, BlueprintReadOnly, Category = "Properties")
 	float FadeOutDelay;
+	
+	UPROPERTY(EditAnywhere, BlueprintReadOnly, Category = "Properties")
+	EWidgetClipping Clipping = EWidgetClipping::OnDemand;
 };
 
 /**
