Slate invalidation: fix CommonLazyImage not showing some async-loaded textures in cooked.

diff --git a/Engine/Plugins/Runtime/CommonUI/Source/CommonUI/Private/CommonLoadGuard.cpp b/Engine/Plugins/Runtime/CommonUI/Source/CommonUI/Private/CommonLoadGuard.cpp
--- a/Engine/Plugins/Runtime/CommonUI/Source/CommonUI/Private/CommonLoadGuard.cpp
+++ b/Engine/Plugins/Runtime/CommonUI/Source/CommonUI/Private/CommonLoadGuard.cpp
@@ -151,6 +151,7 @@ void SLoadGuard::UpdateLoadingAppearance()
 		OnLoadingStateChanged.ExecuteIfBound(false);
 		ContentBorder->SetVisibility(EVisibility::SelfHitTestInvisible);
 		GuardBorder->SetVisibility(EVisibility::Collapsed);
+		Invalidate(EInvalidateWidgetReason::Layout);
 	}
 	else if (!bIsShowingSpinner && (StreamingHandle.IsValid() || bForceShowSpinner))
 	{
@@ -159,6 +160,7 @@ void SLoadGuard::UpdateLoadingAppearance()
 		OnLoadingStateChanged.ExecuteIfBound(true);
 		ContentBorder->SetVisibility(EVisibility::Collapsed);
 		GuardBorder->SetVisibility(EVisibility::SelfHitTestInvisible);
+		Invalidate(EInvalidateWidgetReason::Layout);
 	}
 }
 
