Fixed an edge case where finding an action domain root node wasn't found even if we have multple action domain roots.

diff --git a/Engine/Plugins/Runtime/CommonUI/Source/CommonUI/Private/Input/CommonUIActionRouterBase.cpp b/Engine/Plugins/Runtime/CommonUI/Source/CommonUI/Private/Input/CommonUIActionRouterBase.cpp
--- a/Engine/Plugins/Runtime/CommonUI/Source/CommonUI/Private/Input/CommonUIActionRouterBase.cpp
+++ b/Engine/Plugins/Runtime/CommonUI/Source/CommonUI/Private/Input/CommonUIActionRouterBase.cpp
@@ -1241,12 +1241,12 @@ UCommonUIActionRouterBase::FPendingWidgetRegistration& UCommonUIActionRouterBase
 
 FActivatableTreeNodePtr UCommonUIActionRouterBase::FindNode(const UCommonActivatableWidget* Widget) const
 {
-	FActivatableTreeNodePtr FoundNode;
 	if (Widget)
 	{
 		const bool bIsModal = Widget->IsModal();
 		for (const FActivatableTreeRootPtr RootNode : RootNodes)
 		{
+			FActivatableTreeNodePtr FoundNode;
 			if (!bIsModal)
 			{
 				FoundNode = FindNodeRecursive(RootNode, *Widget);
@@ -1259,27 +1259,24 @@ FActivatableTreeNodePtr UCommonUIActionRouterBase::FindNode(const UCommonActivat
 
 			if (FoundNode.IsValid())
 			{
-				break;
+				return FoundNode;
 			}
 		}
-
-		if (!FoundNode.IsValid())
+				
+		for (const TPair<TObjectPtr<UCommonInputActionDomain>, FActionDomainSortedRootList>& Pair : ActionDomainRootNodes)
 		{
-			for (const TPair<TObjectPtr<UCommonInputActionDomain>, FActionDomainSortedRootList>& Pair : ActionDomainRootNodes)
+			for (const FActivatableTreeRootRef& RootNode : Pair.Value.RootList)
 			{
-				for (const FActivatableTreeRootRef& RootNode : Pair.Value.RootList)
+				FActivatableTreeNodePtr FoundNode = FindNodeRecursive(RootNode, *Widget);
+				if (FoundNode.IsValid())
 				{
-					FoundNode = FindNodeRecursive(RootNode, *Widget);
-					if (FoundNode.IsValid())
-					{
-						break;
-					}
+					return FoundNode;
 				}
 			}
 		}
 	}
 
-	return FoundNode;
+	return nullptr;
 }
 
 FActivatableTreeNodePtr UCommonUIActionRouterBase::FindOwningNode(const UWidget& Widget) const
