Fix an issue with commonbuttongroupbase not selecting a button when button at index 0 is removed.

diff --git a/Engine/Plugins/Runtime/CommonUI/Source/CommonUI/Private/CommonTabListWidgetBase.cpp b/Engine/Plugins/Runtime/CommonUI/Source/CommonUI/Private/CommonTabListWidgetBase.cpp
--- a/Engine/Plugins/Runtime/CommonUI/Source/CommonUI/Private/CommonTabListWidgetBase.cpp
+++ b/Engine/Plugins/Runtime/CommonUI/Source/CommonUI/Private/CommonTabListWidgetBase.cpp
@@ -155,9 +155,24 @@ bool UCommonTabListWidgetBase::RemoveTab(FName TabNameID)
 
 void UCommonTabListWidgetBase::RemoveAllTabs()
 {
+	// We don't call RemoveTab_Internal(Iter->Key, Iter->Value); because that would individually remove TabButtonGroup buttons one by one
+	// Which is something we don't want when we are removing all tabs
+	if (TabButtonGroup)
+	{
+		TabButtonGroup->RemoveAll();
+	}
+	
 	for (TMap<FName, FCommonRegisteredTabInfo>::TIterator Iter(RegisteredTabsByID); Iter; ++Iter)
 	{
-		RemoveTab_Internal(Iter->Key, Iter->Value);
+		if (UCommonButtonBase* const TabButton =  Iter->Value.TabButton)
+		{
+			TabButton->RemoveFromParent();
+
+			RegisteredTabsByID.Remove(Iter->Key);
+			
+			HandleTabRemoval(Iter->Key, TabButton);
+			OnTabButtonRemoval.Broadcast(Iter->Key, TabButton);
+		}
 	}
 }
 
@@ -522,4 +537,4 @@ void UCommonTabListWidgetBase::RemoveTab_Internal(const FName TabNameID, const F
 	// Callbacks
 	HandleTabRemoval(TabNameID, TabButton);
 	OnTabButtonRemoval.Broadcast(TabNameID, TabButton);
-}
+}
\ No newline at end of file
diff --git a/Engine/Plugins/Runtime/CommonUI/Source/CommonUI/Private/Groups/CommonButtonGroupBase.cpp b/Engine/Plugins/Runtime/CommonUI/Source/CommonUI/Private/Groups/CommonButtonGroupBase.cpp
--- a/Engine/Plugins/Runtime/CommonUI/Source/CommonUI/Private/Groups/CommonButtonGroupBase.cpp
+++ b/Engine/Plugins/Runtime/CommonUI/Source/CommonUI/Private/Groups/CommonButtonGroupBase.cpp
@@ -237,6 +237,8 @@ void UCommonButtonGroupBase::OnWidgetRemoved( UWidget* OldWidget )
 
 		if (ButtonIndex == SelectedButtonIndex)
 		{
+			SelectedButtonIndex = INDEX_NONE;
+			
 			if (bSelectionRequired && Buttons.Num() > 0)
 			{
 				for (int32 NewButtonIndex = 0; NewButtonIndex < Buttons.Num(); NewButtonIndex++)
@@ -248,13 +250,10 @@ void UCommonButtonGroupBase::OnWidgetRemoved( UWidget* OldWidget )
 						return; // Early out to only select one button
 					}
 				}
-
-				SelectedButtonIndex = INDEX_NONE;
 				ensureMsgf(false, TEXT("Button group requires selection, but no button is selectable"));
 			}
 			else
 			{
-				SelectedButtonIndex = INDEX_NONE;
 				OnSelectionCleared.Broadcast();
 			}
 		}
