Added an event that is sent when action bar is updated.

diff --git a/Engine/Plugins/Runtime/CommonUI/Source/CommonUI/Private/Input/CommonBoundActionBar.cpp b/Engine/Plugins/Runtime/CommonUI/Source/CommonUI/Private/Input/CommonBoundActionBar.cpp
--- a/Engine/Plugins/Runtime/CommonUI/Source/CommonUI/Private/Input/CommonBoundActionBar.cpp
+++ b/Engine/Plugins/Runtime/CommonUI/Source/CommonUI/Private/Input/CommonBoundActionBar.cpp
@@ -326,6 +326,8 @@ void UCommonBoundActionBar::HandleDeferredDisplayUpdate()
 			}
 		}
 	}
+
+	OnActionBarUpdated.Broadcast();
 }
 
 void UCommonBoundActionBar::HandlePlayerAdded(int32 PlayerIdx)
diff --git a/Engine/Plugins/Runtime/CommonUI/Source/CommonUI/Public/Input/CommonBoundActionBar.h b/Engine/Plugins/Runtime/CommonUI/Source/CommonUI/Public/Input/CommonBoundActionBar.h
--- a/Engine/Plugins/Runtime/CommonUI/Source/CommonUI/Public/Input/CommonBoundActionBar.h
+++ b/Engine/Plugins/Runtime/CommonUI/Source/CommonUI/Public/Input/CommonBoundActionBar.h
@@ -13,6 +13,8 @@ class IWidgetCompilerLog;
 
 struct FUIActionBindingHandle;
 
+DECLARE_DYNAMIC_MULTICAST_DELEGATE(FActionBarUpdated);
+
 /**
  * A box populated with current actions available per CommonUI's Input Handler.
  */
@@ -59,5 +61,8 @@ private:
 	UPROPERTY(EditAnywhere, AdvancedDisplay, Category = Display)
 	bool bIgnoreDuplicateActions = true;
 
+	UPROPERTY(BlueprintAssignable, Category = "Events", meta = (AllowPrivateAccess = true))
+	FActionBarUpdated OnActionBarUpdated;
+
 	bool bIsRefreshQueued = false;
 };
\ No newline at end of file
