Animated Switcher:  expose active widget changed event to BP

diff --git a/Engine/Plugins/Runtime/CommonUI/Source/CommonUI/Private/CommonAnimatedSwitcher.cpp b/Engine/Plugins/Runtime/CommonUI/Source/CommonUI/Private/CommonAnimatedSwitcher.cpp
--- a/Engine/Plugins/Runtime/CommonUI/Source/CommonUI/Private/CommonAnimatedSwitcher.cpp
+++ b/Engine/Plugins/Runtime/CommonUI/Source/CommonUI/Private/CommonAnimatedSwitcher.cpp
@@ -104,6 +104,7 @@ void UCommonAnimatedSwitcher::HandleSlateActiveIndexChanged(int32 ActiveIndex)
 	if (Slots.IsValidIndex(ActiveWidgetIndex))
 	{
 		OnActiveWidgetIndexChanged.Broadcast(GetWidgetAtIndex(ActiveWidgetIndex), ActiveWidgetIndex);
+		OnActiveWidgetIndexChangedBP.Broadcast(GetWidgetAtIndex(ActiveWidgetIndex), ActiveWidgetIndex);
 	}
 }
 PRAGMA_ENABLE_DEPRECATION_WARNINGS
diff --git a/Engine/Plugins/Runtime/CommonUI/Source/CommonUI/Public/CommonAnimatedSwitcher.h b/Engine/Plugins/Runtime/CommonUI/Source/CommonUI/Public/CommonAnimatedSwitcher.h
--- a/Engine/Plugins/Runtime/CommonUI/Source/CommonUI/Public/CommonAnimatedSwitcher.h
+++ b/Engine/Plugins/Runtime/CommonUI/Source/CommonUI/Public/CommonAnimatedSwitcher.h
@@ -10,6 +10,9 @@
 class SOverlay;
 class SSpacer;
 
+
+DECLARE_DYNAMIC_MULTICAST_DELEGATE_TwoParams(FOnActiveIndexChangedDelegate, UWidget*, ActiveWidget, int32, ActiveIndex);
+
 /**
  * A widget switcher that activates / deactivates CommonActivatableWidgets, allowing for associated animations to trigger.
  */
@@ -66,6 +69,10 @@ public:
 	/** Fires when the switcher changes its transition animation state */
 	DECLARE_EVENT_OneParam(UCommonAnimatedSwitcher, FOnTransitioningChanged, bool)
 	FOnTransitioningChanged OnTransitioningChanged;
+	
+	/** Fires when the active widget displayed by the switcher changes */
+	UPROPERTY(BlueprintAssignable, Category = "Common Widget Switcher")
+	FOnActiveIndexChangedDelegate OnActiveWidgetIndexChangedBP;
 
 protected:
 	/** The type of transition to play between widgets */
