PR #10734: Make CommonUI nicer to use for mouse-based strategy games.

diff --git a/Engine/Plugins/Runtime/CommonUI/Source/CommonUI/Private/Input/CommonUIActionRouterBase.cpp b/Engine/Plugins/Runtime/CommonUI/Source/CommonUI/Private/Input/CommonUIActionRouterBase.cpp
--- a/Engine/Plugins/Runtime/CommonUI/Source/CommonUI/Private/Input/CommonUIActionRouterBase.cpp
+++ b/Engine/Plugins/Runtime/CommonUI/Source/CommonUI/Private/Input/CommonUIActionRouterBase.cpp
@@ -1456,11 +1456,11 @@ void UCommonUIActionRouterBase::ApplyUIInputConfig(const FUIInputConfig& NewConf
 						GetWorld()->GetTimerManager().SetTimerForNextTick(this, &ThisClass::FlushPressedKeys);
 					}
 
-					EMouseCaptureMode PrevCaptureMode = GameViewportClient->GetMouseCaptureMode();
-					const bool bWasPermanentlyCaptured = PrevCaptureMode == EMouseCaptureMode::CapturePermanently || PrevCaptureMode == EMouseCaptureMode::CapturePermanently_IncludingInitialMouseDown;
+					const bool bWasCursorHidden = !PC->ShouldShowMouseCursor();
 
 					GameViewportClient->SetMouseCaptureMode(NewConfig.GetMouseCaptureMode());
 					GameViewportClient->SetHideCursorDuringCapture(NewConfig.HideCursorDuringViewportCapture() && !ShouldAlwaysShowCursor());
+					GameViewportClient->SetMouseLockMode(NewConfig.GetMouseLockMode());
 
 					FReply& SlateOperations = LocalPlayer.GetSlateOperations();
 					const EMouseCaptureMode CaptureMode = NewConfig.GetMouseCaptureMode();
@@ -1469,60 +1469,74 @@ void UCommonUIActionRouterBase::ApplyUIInputConfig(const FUIInputConfig& NewConf
 					case EMouseCaptureMode::CapturePermanently:
 					case EMouseCaptureMode::CapturePermanently_IncludingInitialMouseDown:
 					{
-						GameViewportClient->SetMouseLockMode(EMouseLockMode::LockOnCapture);
-						PC->SetShowMouseCursor(ShouldAlwaysShowCursor());
+						PC->SetShowMouseCursor(ShouldAlwaysShowCursor() || !NewConfig.HideCursorDuringViewportCapture());
 
 						TSharedRef<SViewport> ViewportWidgetRef = ViewportWidget.ToSharedRef();
 						SlateOperations.UseHighPrecisionMouseMovement(ViewportWidgetRef);
 						SlateOperations.SetUserFocus(ViewportWidgetRef);
-						SlateOperations.LockMouseToWidget(ViewportWidgetRef);
 						SlateOperations.CaptureMouse(ViewportWidgetRef);
+
+						if (GameViewportClient->ShouldAlwaysLockMouse() || GameViewportClient->LockDuringCapture() || !PC->ShouldShowMouseCursor())
+						{
+							SlateOperations.LockMouseToWidget(ViewportWidget.ToSharedRef());
+						}
+						else
+						{
+							SlateOperations.ReleaseMouseLock();
+						}
 					}
 					break;
 					case EMouseCaptureMode::NoCapture:
 					case EMouseCaptureMode::CaptureDuringMouseDown:
 					case EMouseCaptureMode::CaptureDuringRightMouseDown:
 					{
-						GameViewportClient->SetMouseLockMode(EMouseLockMode::DoNotLock);
 						PC->SetShowMouseCursor(true);
 
-						SlateOperations.ReleaseMouseLock();
 						SlateOperations.ReleaseMouseCapture();
 
-						// If the mouse was captured previously, set it back to the center of the viewport now that we're showing it again 
-						if (!bForceRefresh && bWasPermanentlyCaptured)
+						if (GameViewportClient->ShouldAlwaysLockMouse())
 						{
-							const ECommonInputType CurrentInputType = GetInputSubsystem().GetCurrentInputType();
-							
-							bool bCenterCursor = true;
-							switch (CurrentInputType)
-							{
-								// Touch - Don't do it - the cursor isn't really relevant there.
-								case ECommonInputType::Touch:
-									bCenterCursor = false;
-									break;
-								// Gamepad - Let the settings tell us if we should center it.
-								case ECommonInputType::Gamepad:
-									break;
-							}
+							SlateOperations.LockMouseToWidget(ViewportWidget.ToSharedRef());
+						}
+						else
+						{
+							SlateOperations.ReleaseMouseLock();
+						}
+					}
+					break;
+					}
+
+					// If the mouse was hidden previously, set it back to the center of the viewport now that we're showing it again 
+					if (!bForceRefresh && bWasCursorHidden && PC->ShouldShowMouseCursor())
+					{
+						const ECommonInputType CurrentInputType = GetInputSubsystem().GetCurrentInputType();
+						
+						bool bCenterCursor = true;
+						switch (CurrentInputType)
+						{
+							// Touch - Don't do it - the cursor isn't really relevant there.
+							case ECommonInputType::Touch:
+								bCenterCursor = false;
+								break;
+							// Gamepad - Let the settings tell us if we should center it.
+							case ECommonInputType::Gamepad:
+								break;
+						}
 
-							if (bCenterCursor)
+						if (bCenterCursor)
+						{
+							TSharedPtr<FSlateUser> SlateUser = LocalPlayer.GetSlateUser();
+							TSharedPtr<IGameLayerManager> GameLayerManager = GameViewportClient->GetGameLayerManager();
+							if (ensure(SlateUser) && ensure(GameLayerManager))
 							{
-								TSharedPtr<FSlateUser> SlateUser = LocalPlayer.GetSlateUser();
-								TSharedPtr<IGameLayerManager> GameLayerManager = GameViewportClient->GetGameLayerManager();
-								if (ensure(SlateUser) && ensure(GameLayerManager))
-								{
-									FGeometry PlayerViewGeometry = GameLayerManager->GetPlayerWidgetHostGeometry(&LocalPlayer);
-									const FVector2D AbsoluteViewCenter = PlayerViewGeometry.GetAbsolutePositionAtCoordinates(FVector2D(0.5f, 0.5f));
-									SlateUser->SetCursorPosition(AbsoluteViewCenter);
+								FGeometry PlayerViewGeometry = GameLayerManager->GetPlayerWidgetHostGeometry(&LocalPlayer);
+								const FVector2D AbsoluteViewCenter = PlayerViewGeometry.GetAbsolutePositionAtCoordinates(FVector2D(0.5f, 0.5f));
+								SlateUser->SetCursorPosition(AbsoluteViewCenter);
 
-									UE_LOG(LogUIActionRouter, Verbose, TEXT("Capturing the cursor at the viewport center."));
-								}
+								UE_LOG(LogUIActionRouter, Verbose, TEXT("Moving the cursor to the viewport center."));
 							}
 						}
 					}
-					break;
-					}
 				}
 				else
 				{
diff --git a/Engine/Plugins/Runtime/CommonUI/Source/CommonUI/Private/Input/UIActionRouterTypes.cpp b/Engine/Plugins/Runtime/CommonUI/Source/CommonUI/Private/Input/UIActionRouterTypes.cpp
--- a/Engine/Plugins/Runtime/CommonUI/Source/CommonUI/Private/Input/UIActionRouterTypes.cpp
+++ b/Engine/Plugins/Runtime/CommonUI/Source/CommonUI/Private/Input/UIActionRouterTypes.cpp
@@ -577,6 +577,31 @@ const UWidget* FUIActionBindingHandle::GetBoundWidget() const
 FUIInputConfig::FUIInputConfig()
 	: InputMode(ECommonInputMode::Menu)
 	, MouseCaptureMode(EMouseCaptureMode::NoCapture)
+	, MouseLockMode(EMouseLockMode::DoNotLock)
+{}
+
+FUIInputConfig::FUIInputConfig(ECommonInputMode InInputMode, EMouseCaptureMode InMouseCaptureMode, bool bInHideCursorDuringViewportCapture)
+	: InputMode(InInputMode)
+	, MouseCaptureMode(InMouseCaptureMode)
+	, bHideCursorDuringViewportCapture(bInHideCursorDuringViewportCapture)
+{
+	switch (MouseCaptureMode)
+	{
+	case EMouseCaptureMode::CapturePermanently:
+	case EMouseCaptureMode::CapturePermanently_IncludingInitialMouseDown:
+		MouseLockMode = EMouseLockMode::LockOnCapture;
+		break;
+	default:
+		MouseLockMode = EMouseLockMode::DoNotLock;
+		break;
+	}
+}
+
+FUIInputConfig::FUIInputConfig(ECommonInputMode InInputMode, EMouseCaptureMode InMouseCaptureMode, EMouseLockMode InMouseLockMode, bool bInHideCursorDuringViewportCapture)
+	: InputMode(InInputMode)
+	, MouseCaptureMode(InMouseCaptureMode)
+	, MouseLockMode(InMouseLockMode)
+	, bHideCursorDuringViewportCapture(bInHideCursorDuringViewportCapture)
 {}
 
 FString FUIInputConfig::ToString() const
diff --git a/Engine/Plugins/Runtime/CommonUI/Source/CommonUI/Public/Input/UIActionBindingHandle.h b/Engine/Plugins/Runtime/CommonUI/Source/CommonUI/Public/Input/UIActionBindingHandle.h
--- a/Engine/Plugins/Runtime/CommonUI/Source/CommonUI/Public/Input/UIActionBindingHandle.h
+++ b/Engine/Plugins/Runtime/CommonUI/Source/CommonUI/Public/Input/UIActionBindingHandle.h
@@ -7,6 +7,7 @@
 #include "UIActionBindingHandle.generated.h"
 
 enum class EMouseCaptureMode : uint8;
+enum class EMouseLockMode : uint8;
 struct FScriptContainerElement;
 
 class UWidget;
@@ -89,19 +90,18 @@ struct COMMONUI_API FUIInputConfig
 
 	ECommonInputMode GetInputMode() const { return InputMode; }
 	EMouseCaptureMode GetMouseCaptureMode() const { return MouseCaptureMode; }
+	EMouseLockMode GetMouseLockMode() const { return MouseLockMode; }
 	bool HideCursorDuringViewportCapture() const { return bHideCursorDuringViewportCapture; }
 
 	FUIInputConfig();
-	FUIInputConfig(ECommonInputMode InInputMode, EMouseCaptureMode InMouseCaptureMode, bool bInHideCursorDuringViewportCapture = true)
-		: InputMode(InInputMode)
-		, MouseCaptureMode(InMouseCaptureMode)
-		, bHideCursorDuringViewportCapture(bInHideCursorDuringViewportCapture)
-	{}
+	FUIInputConfig(ECommonInputMode InInputMode, EMouseCaptureMode InMouseCaptureMode, bool bInHideCursorDuringViewportCapture = true);
+	FUIInputConfig(ECommonInputMode InInputMode, EMouseCaptureMode InMouseCaptureMode, EMouseLockMode InMouseLockMode, bool bInHideCursorDuringViewportCapture = true);
 
 	bool operator==(const FUIInputConfig& Other) const
 	{
 		return InputMode == Other.InputMode
 			&& MouseCaptureMode == Other.MouseCaptureMode
+			&& MouseLockMode == Other.MouseLockMode
 			&& bHideCursorDuringViewportCapture == Other.bHideCursorDuringViewportCapture;
 	}
 
@@ -128,6 +128,10 @@ protected:
 	UPROPERTY(EditAnywhere, BlueprintReadWrite, Category = InputConfig)
 	EMouseCaptureMode MouseCaptureMode;
 
+
+	UPROPERTY(EditAnywhere, BlueprintReadWrite, Category = InputConfig)
+	EMouseLockMode MouseLockMode;
+
 	UPROPERTY(EditAnywhere, BlueprintReadWrite, Category = InputConfig)
 	bool bHideCursorDuringViewportCapture = true;
 };
