Added BlueprintCallable function AddWidgets to UCommonWidgetGroupBase.

diff --git a/Engine/Plugins/Runtime/CommonUI/Source/CommonUI/Private/Groups/CommonWidgetGroupBase.cpp b/Engine/Plugins/Runtime/CommonUI/Source/CommonUI/Private/Groups/CommonWidgetGroupBase.cpp
--- a/Engine/Plugins/Runtime/CommonUI/Source/CommonUI/Private/Groups/CommonWidgetGroupBase.cpp
+++ b/Engine/Plugins/Runtime/CommonUI/Source/CommonUI/Private/Groups/CommonWidgetGroupBase.cpp
@@ -16,6 +16,14 @@ void UCommonWidgetGroupBase::AddWidget(UWidget* InWidget)
 	}
 }
 
+void UCommonWidgetGroupBase::AddWidgets(const TArray<UWidget*>& Widgets)
+{
+	for (UWidget* Widget : Widgets)
+	{
+		AddWidget(Widget);
+	}
+}
+
 void UCommonWidgetGroupBase::RemoveWidget(UWidget* InWidget)
 {
 	if (ensure(InWidget) && InWidget->IsA(GetWidgetType()))
diff --git a/Engine/Plugins/Runtime/CommonUI/Source/CommonUI/Public/Groups/CommonWidgetGroupBase.h b/Engine/Plugins/Runtime/CommonUI/Source/CommonUI/Public/Groups/CommonWidgetGroupBase.h
--- a/Engine/Plugins/Runtime/CommonUI/Source/CommonUI/Public/Groups/CommonWidgetGroupBase.h
+++ b/Engine/Plugins/Runtime/CommonUI/Source/CommonUI/Public/Groups/CommonWidgetGroupBase.h
@@ -21,6 +21,9 @@ public:
 	UFUNCTION(BlueprintCallable, Category = Group)
 	void AddWidget(UWidget* InWidget);
 
+	UFUNCTION(BlueprintCallable, Category = Group)
+	void AddWidgets(const TArray<UWidget*>& Widgets);
+
 	UFUNCTION(BlueprintCallable, Category = Group)
 	void RemoveWidget(UWidget* InWidget);
 
@@ -40,4 +43,4 @@ protected:
 	virtual void OnWidgetAdded(UWidget* NewWidget) PURE_VIRTUAL(UCommonWidgetGroupBase::OnWidgetAdded, );
 	virtual void OnWidgetRemoved(UWidget* OldWidget) PURE_VIRTUAL(UCommonWidgetGroupBase::OnWidgetRemoved, );
 	virtual void OnRemoveAll() PURE_VIRTUAL(UCommonWidgetGroupBase::OnRemoveAll, );
-};
\ No newline at end of file
+};
