Added action domain root node scroll recipient based on sorted order to allow for widgets within these root nodes to scroll.
#rnx
#tests: PIE: ran through case of widgets in action domain root nodes as registerd scroll recipients

diff --git a/Engine/Plugins/Runtime/CommonUI/Source/CommonUI/Private/Input/CommonUIActionRouterBase.cpp b/Engine/Plugins/Runtime/CommonUI/Source/CommonUI/Private/Input/CommonUIActionRouterBase.cpp
--- a/Engine/Plugins/Runtime/CommonUI/Source/CommonUI/Private/Input/CommonUIActionRouterBase.cpp
+++ b/Engine/Plugins/Runtime/CommonUI/Source/CommonUI/Private/Input/CommonUIActionRouterBase.cpp
@@ -340,6 +340,27 @@ TArray<const UWidget*> UCommonUIActionRouterBase::GatherActiveAnalogScrollRecipi
 	{
 		return ActiveRootNode->GatherScrollRecipients();
 	}
+	else
+	{
+		if (const UCommonInputActionDomainTable* ActionDomainTable = GetActionDomainTable())
+		{
+			for (const UCommonInputActionDomain* ActionDomain : ActionDomainTable->ActionDomains)
+			{
+				if (const FActionDomainSortedRootList* SortedRootList = ActionDomainRootNodes.Find(ActionDomain))
+				{
+					for (const FActivatableTreeRootRef& RootNode : SortedRootList->RootList)
+					{
+						// only return the first of the root nodes that's in the sorted root list in action domain order.
+						if (RootNode->IsReceivingInput() && RootNode->DoesWidgetSupportActivationFocus())
+						{
+							return RootNode->GatherScrollRecipients();
+						}
+					}
+				}
+			}
+		}
+	}
+
 	return TArray<const UWidget*>();
 }
 
@@ -677,8 +698,7 @@ void UCommonUIActionRouterBase::HandleRootNodeActivated(TWeakPtr<FActivatableTre
 			ActivatedRoot->SetCanReceiveInput(true);
 			ActivatedRoot->OnLeafmostActiveNodeChanged.BindUObject(this, &UCommonUIActionRouterBase::HandleLeafmostActiveNodeChanged);
 			
-			UCommonInputSubsystem& CommonInputSubsystem = GetInputSubsystem();
-			if (UCommonInputActionDomainTable* ActionDomainTable = CommonInputSubsystem.GetActionDomainTable())
+			if (const UCommonInputActionDomainTable* ActionDomainTable = GetActionDomainTable())
 			{
 				// We find the first root node that is receiving input and bail early. 
 				// The action domains and root lists are sorted so we should end up with a node with a higher paint layer
@@ -781,11 +801,15 @@ void UCommonUIActionRouterBase::HandlePostGarbageCollect()
 	}
 }
 
-bool UCommonUIActionRouterBase::ProcessInputOnActionDomains(ECommonInputMode ActiveInputMode, FKey Key, EInputEvent InputEvent) const
+const UCommonInputActionDomainTable* UCommonUIActionRouterBase::GetActionDomainTable() const
 {
 	UCommonInputSubsystem& CommonInputSubsystem = GetInputSubsystem();
-	UCommonInputActionDomainTable* ActionDomainTable = CommonInputSubsystem.GetActionDomainTable();
+	return CommonInputSubsystem.GetActionDomainTable();
+}
 
+bool UCommonUIActionRouterBase::ProcessInputOnActionDomains(ECommonInputMode ActiveInputMode, FKey Key, EInputEvent InputEvent) const
+{
+	const UCommonInputActionDomainTable* ActionDomainTable = GetActionDomainTable();
 	if (!ActionDomainTable)
 	{
 		return false;
@@ -834,10 +858,8 @@ bool UCommonUIActionRouterBase::ProcessInputOnActionDomains(ECommonInputMode Act
 EProcessHoldActionResult UCommonUIActionRouterBase::ProcessHoldInputOnActionDomains(ECommonInputMode ActiveInputMode, FKey Key, EInputEvent InputEvent) const
 {
 	EProcessHoldActionResult HoldActionResult = EProcessHoldActionResult::Unhandled;
-
-	UCommonInputSubsystem& CommonInputSubsystem = GetInputSubsystem();
-	UCommonInputActionDomainTable* ActionDomainTable = CommonInputSubsystem.GetActionDomainTable();
-
+	
+	const UCommonInputActionDomainTable* ActionDomainTable = GetActionDomainTable();
 	if (!ActionDomainTable)
 	{
 		return HoldActionResult;
@@ -1388,32 +1410,29 @@ void UCommonUIActionRouterBase::RefreshActionDomainLeafNodeConfig()
 		return;
 	}
 
-	if (UCommonInputSubsystem* CommonInputSubsystem = GetLocalPlayer() ? GetLocalPlayer()->GetSubsystem<UCommonInputSubsystem>() : nullptr)
+	if (const UCommonInputActionDomainTable* ActionDomainTable = GetActionDomainTable())
 	{
-		if (UCommonInputActionDomainTable* ActionDomainTable = CommonInputSubsystem->GetActionDomainTable())
+		for (const UCommonInputActionDomain* ActionDomain : ActionDomainTable->ActionDomains)
 		{
-			for (const UCommonInputActionDomain* ActionDomain : ActionDomainTable->ActionDomains)
+			if (FActionDomainSortedRootList* SortedRootList = ActionDomainRootNodes.Find(ActionDomain))
 			{
-				if (FActionDomainSortedRootList* SortedRootList = ActionDomainRootNodes.Find(ActionDomain))
+				for (FActivatableTreeRootRef& RootNode : SortedRootList->RootList)
 				{
-					for (FActivatableTreeRootRef& RootNode : SortedRootList->RootList)
+					// only root nodes that are actively receiving input and supports widget activation focus
+					// should update leaf nodes and have input config applied. This will also update focus.
+					if (RootNode->IsReceivingInput() && RootNode->DoesWidgetSupportActivationFocus())
 					{
-						// only root nodes that are actively receiving input and supports widget activation focus
-						// should update leaf nodes and have input config applied. This will also update focus.
-						if (RootNode->IsReceivingInput() && RootNode->DoesWidgetSupportActivationFocus())
+						if (!RootNode->UpdateLeafmostActiveNode(RootNode))
 						{
-							if (!RootNode->UpdateLeafmostActiveNode(RootNode))
-							{
-								RootNode->ApplyLeafmostNodeConfig();
-							}
-							return;
+							RootNode->ApplyLeafmostNodeConfig();
 						}
+						return;
 					}
 				}
 			}
-
-			SetActiveUIInputConfig(FUIInputConfig(ActionDomainTable->InputMode, ActionDomainTable->MouseCaptureMode), ActionDomainTable);
 		}
+
+		SetActiveUIInputConfig(FUIInputConfig(ActionDomainTable->InputMode, ActionDomainTable->MouseCaptureMode), ActionDomainTable);
 	}
 }
 
diff --git a/Engine/Plugins/Runtime/CommonUI/Source/CommonUI/Public/Input/CommonUIActionRouterBase.h b/Engine/Plugins/Runtime/CommonUI/Source/CommonUI/Public/Input/CommonUIActionRouterBase.h
--- a/Engine/Plugins/Runtime/CommonUI/Source/CommonUI/Public/Input/CommonUIActionRouterBase.h
+++ b/Engine/Plugins/Runtime/CommonUI/Source/CommonUI/Public/Input/CommonUIActionRouterBase.h
@@ -28,6 +28,7 @@ class UInputComponent;
 class UPlayerInput;
 class FCommonAnalogCursor;
 class IInputProcessor;
+class UCommonInputActionDomainTable;
 
 enum class EProcessHoldActionResult;
 class FActivatableTreeNode;
@@ -178,6 +179,7 @@ private:
 
 	void HandlePostGarbageCollect();
 
+	const UCommonInputActionDomainTable* GetActionDomainTable() const;
 	bool ProcessInputOnActionDomains(ECommonInputMode ActiveInputMode, FKey Key, EInputEvent InputEvent) const;
 	EProcessHoldActionResult ProcessHoldInputOnActionDomains(ECommonInputMode ActiveInputMode, FKey Key, EInputEvent InputEvent) const;
 
