We don't request a reference to Common Input subsystem.
#rnx
#tests: PIE: Tested startup and shutdown using a frontend

diff --git a/Engine/Plugins/Runtime/CommonUI/Source/CommonUI/Private/Input/CommonUIActionRouterBase.cpp b/Engine/Plugins/Runtime/CommonUI/Source/CommonUI/Private/Input/CommonUIActionRouterBase.cpp
--- a/Engine/Plugins/Runtime/CommonUI/Source/CommonUI/Private/Input/CommonUIActionRouterBase.cpp
+++ b/Engine/Plugins/Runtime/CommonUI/Source/CommonUI/Private/Input/CommonUIActionRouterBase.cpp
@@ -803,8 +803,8 @@ void UCommonUIActionRouterBase::HandlePostGarbageCollect()
 
 const UCommonInputActionDomainTable* UCommonUIActionRouterBase::GetActionDomainTable() const
 {
-	UCommonInputSubsystem& CommonInputSubsystem = GetInputSubsystem();
-	return CommonInputSubsystem.GetActionDomainTable();
+	UCommonInputSubsystem* InputSubsytem = GetLocalPlayerChecked()->GetSubsystem<UCommonInputSubsystem>();
+	return InputSubsytem ? InputSubsytem->GetActionDomainTable() : nullptr;
 }
 
 bool UCommonUIActionRouterBase::ProcessInputOnActionDomains(ECommonInputMode ActiveInputMode, FKey Key, EInputEvent InputEvent) const
