UMG: Change the display name of CommonButtonBase events to be more friendly.

diff --git a/Engine/Plugins/Runtime/CommonUI/Source/CommonUI/Public/CommonButtonBase.h b/Engine/Plugins/Runtime/CommonUI/Source/CommonUI/Public/CommonButtonBase.h
--- a/Engine/Plugins/Runtime/CommonUI/Source/CommonUI/Public/CommonButtonBase.h
+++ b/Engine/Plugins/Runtime/CommonUI/Source/CommonUI/Public/CommonButtonBase.h
@@ -844,10 +844,10 @@ protected:
 	UPROPERTY(BlueprintAssignable, Category = "Events", meta = (AllowPrivateAccess = true))
 	FCommonSelectedStateChangedBase OnSelectedChangedBase;
 
-	UPROPERTY(BlueprintAssignable, Category = "Events", meta = (AllowPrivateAccess = true))
+	UPROPERTY(BlueprintAssignable, Category = "Events", meta = (AllowPrivateAccess = true, DisplayName = "OnClicked"))
 	FCommonButtonBaseClicked OnButtonBaseClicked;
 
-	UPROPERTY(BlueprintAssignable, Category = "Events", meta = (AllowPrivateAccess = true))
+	UPROPERTY(BlueprintAssignable, Category = "Events", meta = (AllowPrivateAccess = true, DisplayName = "OnDoubleClicked"))
 	FCommonButtonBaseClicked OnButtonBaseDoubleClicked;
 
 	UPROPERTY(BlueprintAssignable, Category = "Events", meta = (AllowPrivateAccess = true))
