Added support to CommonDateTimeTextBlock for a custom format property

diff --git a/Engine/Plugins/Runtime/CommonUI/Source/CommonUI/Private/CommonDateTimeTextBlock.cpp b/Engine/Plugins/Runtime/CommonUI/Source/CommonUI/Private/CommonDateTimeTextBlock.cpp
--- a/Engine/Plugins/Runtime/CommonUI/Source/CommonUI/Private/CommonDateTimeTextBlock.cpp
+++ b/Engine/Plugins/Runtime/CommonUI/Source/CommonUI/Private/CommonDateTimeTextBlock.cpp
@@ -131,42 +131,57 @@ TOptional<FText> UCommonDateTimeTextBlock::FormatTimespan(const FTimespan& InTim
 {
 	TOptional<FText> TextToSet = TOptional<FText>();
 
-	if (InTimespan.GetTotalSeconds() < 1.0)
+	if (CustomTimespanFormat.IsEmpty())
 	{
-		TextToSet = FText::AsTimespan(InTimespan);
-	}
-	else
-	{
-		const int32 NewDaysCount = InTimespan.GetDays();
-		if (NewDaysCount > 2)
+		if (InTimespan.GetTotalSeconds() < 1.0)
 		{
-			if (NewDaysCount != GetLastDaysCount())
-			{
-				FText TimespanFormatPattern = NSLOCTEXT("CommonDateTimeTextBlock", "DaysFormatText", "{Days} {Days}|plural(one=Day, other=Days)");
-				FFormatNamedArguments TimeArguments;
-				TimeArguments.Add(TEXT("Days"), NewDaysCount);
-				TextToSet = FText::Format(TimespanFormatPattern, TimeArguments);
-			}
+			TextToSet = FText::AsTimespan(InTimespan);
 		}
 		else
 		{
-			const int32 NewHoursCount = (int32)InTimespan.GetTotalHours();
-			if (NewHoursCount > 12)
+			const int32 NewDaysCount = InTimespan.GetDays();
+			if (NewDaysCount > 2)
 			{
-				if (NewHoursCount != GetLastHoursCount())
+				if (NewDaysCount != GetLastDaysCount())
 				{
-					FText TimespanFormatPattern = NSLOCTEXT("CommonDateTimeTextBlock", "HoursFormatText", "{Hours} {Hours}|plural(one=Hour, other=Hours)");
+					FText TimespanFormatPattern = NSLOCTEXT("CommonDateTimeTextBlock", "DaysFormatText", "{Days} {Days}|plural(one=Day, other=Days)");
 					FFormatNamedArguments TimeArguments;
-					TimeArguments.Add(TEXT("Hours"), NewHoursCount);
+					TimeArguments.Add(TEXT("Days"), NewDaysCount);
 					TextToSet = FText::Format(TimespanFormatPattern, TimeArguments);
 				}
 			}
 			else
 			{
-				TextToSet = FText::AsTimespan(InTimespan);
+				const int32 NewHoursCount = (int32)InTimespan.GetTotalHours();
+				if (NewHoursCount > 12)
+				{
+					if (NewHoursCount != GetLastHoursCount())
+					{
+						FText TimespanFormatPattern = NSLOCTEXT("CommonDateTimeTextBlock", "HoursFormatText", "{Hours} {Hours}|plural(one=Hour, other=Hours)");
+						FFormatNamedArguments TimeArguments;
+						TimeArguments.Add(TEXT("Hours"), NewHoursCount);
+						TextToSet = FText::Format(TimespanFormatPattern, TimeArguments);
+					}
+				}
+				else
+				{
+					TextToSet = FText::AsTimespan(InTimespan);
+				}
 			}
 		}
 	}
+	else
+	{
+		FNumberFormattingOptions NumberFormatOptions;
+		NumberFormatOptions.MinimumIntegralDigits = bCustomTimespanLeadingZeros ? 2 : 1;
+
+		FFormatNamedArguments TimeArguments;
+		TimeArguments.Add(TEXT("Days"), FText::AsNumber(InTimespan.GetDays(), &NumberFormatOptions));
+		TimeArguments.Add(TEXT("Hours"), FText::AsNumber(InTimespan.GetHours(), &NumberFormatOptions));
+		TimeArguments.Add(TEXT("Minutes"), FText::AsNumber(InTimespan.GetMinutes(), &NumberFormatOptions));
+		TimeArguments.Add(TEXT("Seconds"), FText::AsNumber(InTimespan.GetSeconds(), &NumberFormatOptions));
+		TextToSet = FText::Format(CustomTimespanFormat, TimeArguments);
+	}
 
 	return TextToSet;
 }
diff --git a/Engine/Plugins/Runtime/CommonUI/Source/CommonUI/Public/CommonDateTimeTextBlock.h b/Engine/Plugins/Runtime/CommonUI/Source/CommonUI/Public/CommonDateTimeTextBlock.h
--- a/Engine/Plugins/Runtime/CommonUI/Source/CommonUI/Public/CommonDateTimeTextBlock.h
+++ b/Engine/Plugins/Runtime/CommonUI/Source/CommonUI/Public/CommonDateTimeTextBlock.h
@@ -46,6 +46,19 @@ protected:
 	int32 GetLastDaysCount() const { return LastDaysCount; }
 	int32 GetLastHoursCount() const { return LastHoursCount; }
 
+	/*
+	* Supplies a custom timespan format to use if desired
+	* Supported arguments include {Days}, {Hours}, {Minutes}, and {Seconds}
+	*/
+	UPROPERTY(EditAnywhere, Category = "Custom Timespan")
+	FText CustomTimespanFormat;
+
+	/*
+	* If the custom timespan should use a leading zero for values, ie "02"
+	*/
+	UPROPERTY(EditAnywhere, Category = "Custom Timespan")
+	bool bCustomTimespanLeadingZeros = false;
+
 private:
 
 	// Timer handle for timer-based ticking based on InterpolationUpdateInterval.
