CommonActivatableWidget::RequestRefreshFocus() function is now exposed to Blueprint

diff --git a/Engine/Plugins/Runtime/CommonUI/Source/CommonUI/Public/CommonActivatableWidget.h b/Engine/Plugins/Runtime/CommonUI/Source/CommonUI/Public/CommonActivatableWidget.h
--- a/Engine/Plugins/Runtime/CommonUI/Source/CommonUI/Public/CommonActivatableWidget.h
+++ b/Engine/Plugins/Runtime/CommonUI/Source/CommonUI/Public/CommonActivatableWidget.h
@@ -138,6 +138,7 @@ protected:
 	 * or the buttons are deeply nested in a multi-switcher hierarchy and it would be burdensome
 	 * to wrap each element in a CommonActivatableWidget.
 	 */
+	UFUNCTION(BlueprintCallable, Category = ActivatableWidget)
 	void RequestRefreshFocus();
 
 	UFUNCTION(BlueprintImplementableEvent, Category = ActivatableWidget, meta = (DisplayName = "On Activated"))
