Updated display names for UCommonButtonBase events.

diff --git a/Engine/Plugins/Runtime/CommonUI/Source/CommonUI/Public/CommonButtonBase.h b/Engine/Plugins/Runtime/CommonUI/Source/CommonUI/Public/CommonButtonBase.h
--- a/Engine/Plugins/Runtime/CommonUI/Source/CommonUI/Public/CommonButtonBase.h
+++ b/Engine/Plugins/Runtime/CommonUI/Source/CommonUI/Public/CommonButtonBase.h
@@ -841,19 +841,19 @@ public:
 #endif
 
 protected:
-	UPROPERTY(BlueprintAssignable, Category = "Events", meta = (AllowPrivateAccess = true))
+	UPROPERTY(BlueprintAssignable, Category = "Events", meta = (AllowPrivateAccess = true, DisplayName = "On Selected Changed"))
 	FCommonSelectedStateChangedBase OnSelectedChangedBase;
 
-	UPROPERTY(BlueprintAssignable, Category = "Events", meta = (AllowPrivateAccess = true, DisplayName = "OnClicked"))
+	UPROPERTY(BlueprintAssignable, Category = "Events", meta = (AllowPrivateAccess = true, DisplayName = "On Clicked"))
 	FCommonButtonBaseClicked OnButtonBaseClicked;
 
-	UPROPERTY(BlueprintAssignable, Category = "Events", meta = (AllowPrivateAccess = true, DisplayName = "OnDoubleClicked"))
+	UPROPERTY(BlueprintAssignable, Category = "Events", meta = (AllowPrivateAccess = true, DisplayName = "On Double Clicked"))
 	FCommonButtonBaseClicked OnButtonBaseDoubleClicked;
 
-	UPROPERTY(BlueprintAssignable, Category = "Events", meta = (AllowPrivateAccess = true))
+	UPROPERTY(BlueprintAssignable, Category = "Events", meta = (AllowPrivateAccess = true, DisplayName = "On Hovered"))
 	FCommonButtonBaseClicked OnButtonBaseHovered;
 
-	UPROPERTY(BlueprintAssignable, Category = "Events", meta = (AllowPrivateAccess = true))
+	UPROPERTY(BlueprintAssignable, Category = "Events", meta = (AllowPrivateAccess = true, DisplayName = "On Unhovered"))
 	FCommonButtonBaseClicked OnButtonBaseUnhovered;
 
 	FUIActionBindingHandle TriggeringBindingHandle;
