PR #10852: Fix bMatchSize for lazy textures in image widgets in editor

diff --git a/Engine/Plugins/Runtime/CommonUI/Source/CommonUI/Private/CommonLazyImage.cpp b/Engine/Plugins/Runtime/CommonUI/Source/CommonUI/Private/CommonLazyImage.cpp
--- a/Engine/Plugins/Runtime/CommonUI/Source/CommonUI/Private/CommonLazyImage.cpp
+++ b/Engine/Plugins/Runtime/CommonUI/Source/CommonUI/Private/CommonLazyImage.cpp
@@ -10,6 +10,10 @@
 #include "Engine/Texture2DDynamic.h"
 #include "Widgets/Images/SImage.h"
 
+#if WITH_EDITOR
+#include "TextureCompiler.h"
+#endif
+
 #include UE_INLINE_GENERATED_CPP_BY_NAME(CommonLazyImage)
 
 UCommonLazyImage::UCommonLazyImage(const FObjectInitializer& Initializer)
@@ -246,6 +250,9 @@ void UCommonLazyImage::SetBrushObjectInternal(UTexture* Texture, bool bMatchSize
 				}
 				else if (UTexture2D* AsTexture2D = Cast<UTexture2D>(Texture))
 				{
+#if WITH_EDITOR
+					FTextureCompilingManager::Get().FinishCompilation({ Texture });
+#endif
 					CurrentBrush.ImageSize.X = AsTexture2D->GetSizeX();
 					CurrentBrush.ImageSize.Y = AsTexture2D->GetSizeY();
 				}
diff --git a/Engine/Source/Runtime/Engine/Private/Texture2D.cpp b/Engine/Source/Runtime/Engine/Private/Texture2D.cpp
--- a/Engine/Source/Runtime/Engine/Private/Texture2D.cpp
+++ b/Engine/Source/Runtime/Engine/Private/Texture2D.cpp
@@ -336,7 +336,7 @@ FTexturePlatformData* UTexture2D::GetPlatformData()
 // placeholders value to ensure rendering works as expected and that
 // there are no thread-unsafe access to the platform data being built.
 // Any process requiring a fully up-to-date platform data is expected to
-// call FTextureCompiler:Get().FinishCompilation on UTexture first.
+// call FTextureCompilingManager:Get().FinishCompilation on UTexture first.
 int32 UTexture2D::GetSizeX() const
 {
 	if (PrivatePlatformData)
diff --git a/Engine/Source/Runtime/UMG/Private/Components/Image.cpp b/Engine/Source/Runtime/UMG/Private/Components/Image.cpp
--- a/Engine/Source/Runtime/UMG/Private/Components/Image.cpp
+++ b/Engine/Source/Runtime/UMG/Private/Components/Image.cpp
@@ -11,6 +11,10 @@
 #include "Engine/AssetManager.h"
 #include "Engine/StreamableManager.h"
 
+#if WITH_EDITOR
+#include "TextureCompiler.h"
+#endif
+
 #include UE_INLINE_GENERATED_CPP_BY_NAME(Image)
 
 #define LOCTEXT_NAMESPACE "UMG"
@@ -187,6 +191,9 @@ void UImage::SetBrushFromTexture(UTexture2D* Texture, bool bMatchSize)
 		{
 			if (Texture)
 			{
+#if WITH_EDITOR
+				FTextureCompilingManager::Get().FinishCompilation({ Texture });
+#endif
 				Brush.ImageSize.X = Texture->GetSizeX();
 				Brush.ImageSize.Y = Texture->GetSizeY();
 			}
