PR #10853: CommonUI: Add support for size matching CommonVideoPlayer to media

diff --git a/Engine/Plugins/Runtime/CommonUI/Source/CommonUI/Private/CommonVideoPlayer.cpp b/Engine/Plugins/Runtime/CommonUI/Source/CommonUI/Private/CommonVideoPlayer.cpp
--- a/Engine/Plugins/Runtime/CommonUI/Source/CommonUI/Private/CommonVideoPlayer.cpp
+++ b/Engine/Plugins/Runtime/CommonUI/Source/CommonUI/Private/CommonVideoPlayer.cpp
@@ -98,6 +98,22 @@ void UCommonVideoPlayer::SetIsMuted(bool bInIsMuted)
 	}
 }
 
+void UCommonVideoPlayer::SetShouldMatchSize(bool bInMatchSize)
+{
+	if (bMatchSize != bInMatchSize)
+	{
+		bMatchSize = bInMatchSize;
+		if (bMatchSize && MediaPlayer->IsReady())
+		{
+			VideoBrush.ImageSize = MediaPlayer->GetVideoTrackDimensions(INDEX_NONE, INDEX_NONE);
+			if (MyImage)
+			{
+				MyImage->InvalidateImage();
+			}
+		}
+	}
+}
+
 void UCommonVideoPlayer::Play()
 {
 	SetPlaybackRate(1.f);
@@ -180,6 +196,17 @@ void UCommonVideoPlayer::HandleMediaPlayerEvent(EMediaEvent EventType)
 {
 	switch (EventType)
 	{
+	case EMediaEvent::TracksChanged:
+	case EMediaEvent::MediaOpened:
+		if (bMatchSize)
+		{
+			VideoBrush.ImageSize = MediaPlayer->GetVideoTrackDimensions(INDEX_NONE, INDEX_NONE);
+			if (MyImage)
+			{
+				MyImage->InvalidateImage();
+			}
+		}
+		break;
 	case EMediaEvent::MediaClosed:
 		Close();
 		break;
diff --git a/Engine/Plugins/Runtime/CommonUI/Source/CommonUI/Public/CommonVideoPlayer.h b/Engine/Plugins/Runtime/CommonUI/Source/CommonUI/Public/CommonVideoPlayer.h
--- a/Engine/Plugins/Runtime/CommonUI/Source/CommonUI/Public/CommonVideoPlayer.h
+++ b/Engine/Plugins/Runtime/CommonUI/Source/CommonUI/Public/CommonVideoPlayer.h
@@ -32,6 +32,7 @@ public:
 	void SetPlaybackRate(float PlaybackRate);
 	void SetLooping(bool bShouldLoopPlayback);
 	void SetIsMuted(bool bInIsMuted);
+	void SetShouldMatchSize(bool bInMatchSize);
 
 	void Play();
 	void Reverse();
@@ -68,6 +69,10 @@ private:
 	UPROPERTY(EditAnywhere, Category = VideoPlayer)
 	TObjectPtr<UMediaSource> Video;
 
+	// Should we match the size of the media source when it opens?
+	UPROPERTY(EditAnywhere, Category = VideoPlayer)
+	bool bMatchSize = false;
+
 	UPROPERTY(Transient)
 	TObjectPtr<UMediaPlayer> MediaPlayer;
 
