PR #10900: Fix not being able to click Common Text widgets in the widget designer when bAutoCollapseWithEmptyText is set to true

diff --git a/Engine/Plugins/Runtime/CommonUI/Source/CommonUI/Private/CommonTextBlock.cpp b/Engine/Plugins/Runtime/CommonUI/Source/CommonUI/Private/CommonTextBlock.cpp
--- a/Engine/Plugins/Runtime/CommonUI/Source/CommonUI/Private/CommonTextBlock.cpp
+++ b/Engine/Plugins/Runtime/CommonUI/Source/CommonUI/Private/CommonTextBlock.cpp
@@ -407,7 +407,14 @@ void UCommonTextBlock::SynchronizeProperties()
 
 	if (bAutoCollapseWithEmptyText)
 	{
-		SetVisibility(GetText().IsEmpty() ? ESlateVisibility::Collapsed : ESlateVisibility::SelfHitTestInvisible);
+		if (IsDesignTime())
+		{
+			SetVisibility(GetText().IsEmpty() ? ESlateVisibility::Collapsed : ESlateVisibility::Visible);
+		}
+		else
+		{
+			SetVisibility(GetText().IsEmpty() ? ESlateVisibility::Collapsed : ESlateVisibility::SelfHitTestInvisible);
+		}
 	}
 
 	if (CommonUIUtils::ShouldDisplayMobileUISizes())
@@ -439,7 +446,14 @@ void UCommonTextBlock::OnTextChanged()
 	Super::OnTextChanged();
 	if (bAutoCollapseWithEmptyText)
 	{
-		SetVisibility(GetText().IsEmpty() ? ESlateVisibility::Collapsed : ESlateVisibility::SelfHitTestInvisible);
+		if (IsDesignTime())
+		{
+			SetVisibility(GetText().IsEmpty() ? ESlateVisibility::Collapsed : ESlateVisibility::Visible);
+		}
+		else
+		{
+			SetVisibility(GetText().IsEmpty() ? ESlateVisibility::Collapsed : ESlateVisibility::SelfHitTestInvisible);
+		}
 	}
 }
 
