PR #10686: Add bIgnoreMoveInput and bIgnoreLookInput to FUIInputConfig equality

diff --git a/Engine/Plugins/Runtime/CommonUI/Source/CommonUI/Public/Input/UIActionBindingHandle.h b/Engine/Plugins/Runtime/CommonUI/Source/CommonUI/Public/Input/UIActionBindingHandle.h
--- a/Engine/Plugins/Runtime/CommonUI/Source/CommonUI/Public/Input/UIActionBindingHandle.h
+++ b/Engine/Plugins/Runtime/CommonUI/Source/CommonUI/Public/Input/UIActionBindingHandle.h
@@ -99,7 +99,9 @@ struct COMMONUI_API FUIInputConfig
 
 	bool operator==(const FUIInputConfig& Other) const
 	{
-		return InputMode == Other.InputMode
+		return bIgnoreMoveInput == Other.bIgnoreMoveInput
+			&& bIgnoreLookInput == Other.bIgnoreLookInput
+			&& InputMode == Other.InputMode
 			&& MouseCaptureMode == Other.MouseCaptureMode
 			&& MouseLockMode == Other.MouseLockMode
 			&& bHideCursorDuringViewportCapture == Other.bHideCursorDuringViewportCapture;
