PR #10838: Deselect buttons when removing them from a UCommonButtonGroupBase

diff --git a/Engine/Plugins/Runtime/CommonUI/Source/CommonUI/Private/Groups/CommonButtonGroupBase.cpp b/Engine/Plugins/Runtime/CommonUI/Source/CommonUI/Private/Groups/CommonButtonGroupBase.cpp
--- a/Engine/Plugins/Runtime/CommonUI/Source/CommonUI/Private/Groups/CommonButtonGroupBase.cpp
+++ b/Engine/Plugins/Runtime/CommonUI/Source/CommonUI/Private/Groups/CommonButtonGroupBase.cpp
@@ -235,6 +235,12 @@ void UCommonButtonGroupBase::OnWidgetRemoved( UWidget* OldWidget )
 
 		Buttons.RemoveAll( [Button]( TWeakObjectPtr<UCommonButtonBase> Entry ) { return Entry == Button || !Entry.IsValid(); } );
 
+		if (Button->GetSelected())
+		{
+			constexpr bool bAllowSound = false;
+			Button->SetSelectedInternal(false, bAllowSound);
+		}
+
 		if (ButtonIndex == SelectedButtonIndex)
 		{
 			SelectedButtonIndex = INDEX_NONE;
