Speculative fix for a crash when trying to focus the leaf most node in a FActivatableTreeRoot.

This changelist is a speculative fix for a popular crash that happens when trying to focus the least most node of a FActivatableTreeRoot with an invalid LeafmostActiveNode.

I added more checks, error handling and error logging.

#rnx

#tests Made sure the code behaves as before for regular non-error cases.

diff --git a/Engine/Plugins/Runtime/CommonUI/Source/CommonUI/Private/Input/UIActionRouterTypes.cpp b/Engine/Plugins/Runtime/CommonUI/Source/CommonUI/Private/Input/UIActionRouterTypes.cpp
--- a/Engine/Plugins/Runtime/CommonUI/Source/CommonUI/Private/Input/UIActionRouterTypes.cpp
+++ b/Engine/Plugins/Runtime/CommonUI/Source/CommonUI/Private/Input/UIActionRouterTypes.cpp
@@ -1501,10 +1501,27 @@ void FActivatableTreeRoot::ApplyLeafmostNodeConfig()
 void FActivatableTreeRoot::FocusLeafmostNode()
 {
 	check(LeafmostActiveNode.IsValid());
-
+	if (!LeafmostActiveNode.IsValid())
+	{
+		UE_LOG(LogUIActionRouter, Error, TEXT("Cannot focus leaf most node - invalid LeafmostActiveNode"));
+		return;
+	}
 	FActivatableTreeNodePtr PinnedLeafmostNode = LeafmostActiveNode.Pin();
+
+	check(PinnedLeafmostNode);
+	if (!PinnedLeafmostNode)
+	{
+		UE_LOG(LogUIActionRouter, Error, TEXT("Cannot focus leaf most node - invalid PinnedLeafmostNode"));
+		return;
+	}
+
 	UCommonActivatableWidget* LeafWidget = PinnedLeafmostNode->GetWidget();
 	check(LeafWidget);
+	if (!LeafWidget)
+	{
+		UE_LOG(LogUIActionRouter, Error, TEXT("Cannot focus leaf most node - invalid LeafWidget"));
+		return;
+	}
 
 	const int32 OwnerSlateId = GetOwnerUserIndex();
 	ULocalPlayer& LocalPlayer = *GetActionRouter().GetLocalPlayerChecked();
