PR #11135: Fix not being able to click UCommonActionWidgets in the designer

diff --git a/Engine/Plugins/Runtime/CommonUI/Source/CommonUI/Private/CommonActionWidget.cpp b/Engine/Plugins/Runtime/CommonUI/Source/CommonUI/Private/CommonActionWidget.cpp
--- a/Engine/Plugins/Runtime/CommonUI/Source/CommonUI/Private/CommonActionWidget.cpp
+++ b/Engine/Plugins/Runtime/CommonUI/Source/CommonUI/Private/CommonActionWidget.cpp
@@ -343,7 +343,7 @@ void UCommonActionWidget::UpdateActionWidget()
 						}
 
 						MyKeyBox->Invalidate(EInvalidateWidget::LayoutAndVolatility);
-						SetVisibility(ESlateVisibility::SelfHitTestInvisible);
+						SetVisibility(IsDesignTime() ? ESlateVisibility::Visible : ESlateVisibility::SelfHitTestInvisible);
 
 						return;
 					}
