Adding verbose logging to CommonActivatableWidgetContainer to help with debugging.

diff --git a/Engine/Plugins/Runtime/CommonUI/Source/CommonUI/Private/Widgets/CommonActivatableWidgetContainer.cpp b/Engine/Plugins/Runtime/CommonUI/Source/CommonUI/Private/Widgets/CommonActivatableWidgetContainer.cpp
--- a/Engine/Plugins/Runtime/CommonUI/Source/CommonUI/Private/Widgets/CommonActivatableWidgetContainer.cpp
+++ b/Engine/Plugins/Runtime/CommonUI/Source/CommonUI/Private/Widgets/CommonActivatableWidgetContainer.cpp
@@ -2,6 +2,7 @@
 
 #include "Widgets/CommonActivatableWidgetContainer.h"
 #include "CommonActivatableWidget.h"
+#include "CommonUIPrivate.h"
 #include "Slate/SCommonAnimatedSwitcher.h"
 #include "Widgets/SOverlay.h"
 #include "Widgets/Layout/SSpacer.h"
@@ -58,6 +59,8 @@ void UCommonActivatableWidgetContainerBase::RemoveWidget(UCommonActivatableWidge
 
 void UCommonActivatableWidgetContainerBase::RemoveWidget(UCommonActivatableWidget& WidgetToRemove)
 {
+	UE_LOG(LogCommonUI, VeryVerbose, TEXT("UCommonActivatableWidgetContainerBase::RemoveWidget() WidgetToRemove: %s"), *WidgetToRemove.GetName());
+	
 	if (&WidgetToRemove == GetActiveWidget())
 	{
 		// To remove the active widget, just deactivate it (if it's already deactivated, then we're already in the process of ditching it)
@@ -135,6 +138,8 @@ void UCommonActivatableWidgetContainerBase::OnWidgetRebuilt()
 
 void UCommonActivatableWidgetContainerBase::SetSwitcherIndex(int32 TargetIndex, bool bInstantTransition /*= false*/)
 {
+	UE_LOG(LogCommonUI, VeryVerbose, TEXT("UCommonActivatableWidgetContainerBase::SetSwitcherIndex() TargetIndex: %d, bInstantTransition: %d"), TargetIndex, bInstantTransition);
+	
 	if (MySwitcher && MySwitcher->GetActiveWidgetIndex() != TargetIndex)
 	{
 		if (DisplayedWidget)
@@ -174,6 +179,8 @@ UCommonActivatableWidget* UCommonActivatableWidgetContainerBase::AddWidgetIntern
 
 void UCommonActivatableWidgetContainerBase::RegisterInstanceInternal(UCommonActivatableWidget& NewWidget)
 {
+	UE_LOG(LogCommonUI, VeryVerbose, TEXT("UCommonActivatableWidgetContainerBase::RegisterInstanceInternal() NewWidget: %s"), *NewWidget.GetName());
+	
 	// @TODO: Log if bAutoActivate is true on the provided widget, since it quite simply makes no sense.
 	if (ensure(!WidgetList.Contains(&NewWidget)))
 	{
@@ -184,6 +191,8 @@ void UCommonActivatableWidgetContainerBase::RegisterInstanceInternal(UCommonActi
 
 void UCommonActivatableWidgetContainerBase::HandleSwitcherIsTransitioningChanged(bool bIsTransitioning)
 {
+	UE_LOG(LogCommonUI, VeryVerbose, TEXT("UCommonActivatableWidgetContainerBase::HandleSwitcherIsTransitioningChanged() bIsTransitioning: %d"), bIsTransitioning);
+	
 	// While the switcher is transitioning, put up the guard to intercept all input
 	MyInputGuard->SetVisibility(bIsTransitioning ? EVisibility::Visible : EVisibility::Collapsed);
 	OnTransitioningChanged.Broadcast(this, bIsTransitioning);
@@ -191,6 +200,8 @@ void UCommonActivatableWidgetContainerBase::HandleSwitcherIsTransitioningChanged
 
 void UCommonActivatableWidgetContainerBase::HandleActiveWidgetDeactivated(UCommonActivatableWidget* DeactivatedWidget)
 {
+	UE_LOG(LogCommonUI, VeryVerbose, TEXT("UCommonActivatableWidgetContainerBase::HandleActiveWidgetDeactivated() DeactivatedWidget: %s"), *DeactivatedWidget->GetName());
+	
 	// When the currently displayed widget deactivates, transition the switcher to the preceding slot (if it exists)
 	// We'll clean up this slot once the switcher index actually changes
 	if (ensure(DeactivatedWidget == DisplayedWidget) && MySwitcher && MySwitcher->GetActiveWidgetIndex() > 0)
@@ -202,14 +213,25 @@ void UCommonActivatableWidgetContainerBase::HandleActiveWidgetDeactivated(UCommo
 
 void UCommonActivatableWidgetContainerBase::ReleaseWidget(const TSharedRef<SWidget>& WidgetToRelease)
 {
+	UE_LOG(LogCommonUI, VeryVerbose, TEXT("UCommonActivatableWidgetContainerBase::ReleaseWidget()"));
+	
 	if (UCommonActivatableWidget* ActivatableWidget = ActivatableWidgetFromSlate(WidgetToRelease))
 	{
+		UE_LOG(LogCommonUI, VeryVerbose, TEXT("UCommonActivatableWidgetContainerBase WidgetToRelease: %s"), *ActivatableWidget->GetName());
+		
 		GeneratedWidgetsPool.Release(ActivatableWidget, true);
 		WidgetList.Remove(ActivatableWidget);
 	}
+	else
+	{
+		UE_LOG(LogCommonUI, Warning, TEXT("UCommonActivatableWidgetContainerBase::ReleaseWidget() No matching Activatable Widget found."));
+	}
 
-	if (MySwitcher->RemoveSlot(WidgetToRelease) != INDEX_NONE)
+	const int32 RemovedIndex = MySwitcher->RemoveSlot(WidgetToRelease);
+	if (RemovedIndex != INDEX_NONE)
 	{
+		UE_LOG(LogCommonUI, VeryVerbose, TEXT("UCommonActivatableWidgetContainerBase Widget removed from slot %d"), RemovedIndex);
+		
 		ReleasedWidgets.Add(WidgetToRelease);
 		if (ReleasedWidgets.Num() == 1)
 		{
@@ -226,6 +248,8 @@ void UCommonActivatableWidgetContainerBase::ReleaseWidget(const TSharedRef<SWidg
 
 void UCommonActivatableWidgetContainerBase::HandleActiveIndexChanged(int32 ActiveWidgetIndex)
 {
+	UE_LOG(LogCommonUI, VeryVerbose, TEXT("UCommonActivatableWidgetContainerBase::HandleActiveIndexChanged() ActiveWidgetIndex: %d"), ActiveWidgetIndex);
+	
 	// Remove all slots above the currently active one and release the widgets back to the pool
 	while (MySwitcher->GetNumWidgets() - 1 > ActiveWidgetIndex)
 	{
@@ -337,6 +361,8 @@ void UCommonActivatableWidgetQueue::OnWidgetAddedToList(UCommonActivatableWidget
 {
 	if (MySwitcher)
 	{
+		UE_LOG(LogCommonUI, VeryVerbose, TEXT("UCommonActivatableWidgetQueue::OnWidgetAddedToList() PriorNumWidgets: %d"), MySwitcher->GetNumWidgets() - 1);
+		
 		// Insert after the empty slot 0 and before the already queued widgets
 		MySwitcher->AddSlot(1) [AddedWidget.TakeWidget()];
 
