CommonRotator: Add direction to OnRotate event

#11127

diff --git a/Engine/Plugins/Runtime/CommonUI/Source/CommonUI/Private/CommonRotator.cpp b/Engine/Plugins/Runtime/CommonUI/Source/CommonUI/Private/CommonRotator.cpp
--- a/Engine/Plugins/Runtime/CommonUI/Source/CommonUI/Private/CommonRotator.cpp
+++ b/Engine/Plugins/Runtime/CommonUI/Source/CommonUI/Private/CommonRotator.cpp
@@ -103,10 +103,17 @@ void UCommonRotator::ShiftTextLeftInternal(bool bFromNavigation)
 		{
 			SetSelectedItem(SelectedIndex);
 
+			PRAGMA_DISABLE_DEPRECATION_WARNINGS
 			if (OnRotated.IsBound())
 			{
 				OnRotated.Broadcast(SelectedIndex);
 			}
+			PRAGMA_ENABLE_DEPRECATION_WARNINGS
+
+			if (OnRotatedWithDirection.IsBound())
+			{
+				OnRotatedWithDirection.Broadcast(SelectedIndex, ERotatorDirection::Left);
+			}
 
 			OnRotatedEvent.Broadcast(SelectedIndex, bFromNavigation);
 		}
@@ -135,10 +142,17 @@ void UCommonRotator::ShiftTextRightInternal(bool bFromNavigation)
 		{
 			SetSelectedItem(SelectedIndex);
 
+			PRAGMA_DISABLE_DEPRECATION_WARNINGS
 			if (OnRotated.IsBound())
 			{
 				OnRotated.Broadcast(SelectedIndex);
 			}
+			PRAGMA_ENABLE_DEPRECATION_WARNINGS
+
+			if (OnRotatedWithDirection.IsBound())
+			{
+				OnRotatedWithDirection.Broadcast(SelectedIndex, ERotatorDirection::Right);
+			}
 
 			OnRotatedEvent.Broadcast(SelectedIndex, bFromNavigation);
 		}
diff --git a/Engine/Plugins/Runtime/CommonUI/Source/CommonUI/Public/CommonRotator.h b/Engine/Plugins/Runtime/CommonUI/Source/CommonUI/Public/CommonRotator.h
--- a/Engine/Plugins/Runtime/CommonUI/Source/CommonUI/Public/CommonRotator.h
+++ b/Engine/Plugins/Runtime/CommonUI/Source/CommonUI/Public/CommonRotator.h
@@ -10,7 +10,15 @@
 
 class UCommonTextBlock;
 
+UENUM(BlueprintType)
+enum class ERotatorDirection : uint8
+{
+	Right,
+	Left,
+};
+
 DECLARE_DYNAMIC_MULTICAST_DELEGATE_OneParam(FOnRotated, int32, Value);
+DECLARE_DYNAMIC_MULTICAST_DELEGATE_TwoParams(FOnRotatedWithDirection, int32, Value, ERotatorDirection, RotatorDir);
 
 /**
 * A button that can rotate between given text labels.
@@ -21,6 +29,7 @@ class COMMONUI_API UCommonRotator : public UCommonButtonBase
 	GENERATED_UCLASS_BODY()
 
 public:
+
 	virtual bool Initialize() override;
 	virtual FNavigationReply NativeOnNavigation(const FGeometry& MyGeometry, const FNavigationEvent& InNavigationEvent, const FNavigationReply& InDefaultReply);
 
@@ -54,7 +63,12 @@ public:
 
 public:
 
+	/** Called when the Selected state of this button changes. Provides the direction of rotation. */
+	UPROPERTY(BlueprintAssignable, Category = "Events")
+	FOnRotatedWithDirection OnRotatedWithDirection;
+
 	/** Called when the Selected state of this button changes */
+	UE_DEPRECATED(5.4, "OnRotated is deprecated, please use OnRotatedWithDirection instead.")
 	UPROPERTY(BlueprintAssignable, Category = "Events")
 	FOnRotated OnRotated;
 
