Adds text block property ApplyLineHeightToBottomLine

#tests: Tested in Editor

diff --git a/Engine/Plugins/Runtime/CommonUI/Source/CommonUI/Private/CommonTextBlock.cpp b/Engine/Plugins/Runtime/CommonUI/Source/CommonUI/Private/CommonTextBlock.cpp
--- a/Engine/Plugins/Runtime/CommonUI/Source/CommonUI/Private/CommonTextBlock.cpp
+++ b/Engine/Plugins/Runtime/CommonUI/Source/CommonUI/Private/CommonTextBlock.cpp
@@ -29,6 +29,7 @@ FTextScrollerOptions UCommonTextScrollStyle::ToScrollOptions() const
 UCommonTextStyle::UCommonTextStyle()
 	: Color(FLinearColor::Black)
 	, LineHeightPercentage(1.0f)
+	, ApplyLineHeightToBottomLine(true)
 {
 }
 
@@ -52,6 +53,11 @@ float UCommonTextStyle::GetLineHeightPercentage() const
 	return LineHeightPercentage;
 }
 
+bool UCommonTextStyle::GetApplyLineHeightToBottomLine() const
+{
+	return ApplyLineHeightToBottomLine;
+}
+
 void UCommonTextStyle::GetShadowOffset(FVector2D& OutShadowOffset) const
 {
 	OutShadowOffset = ShadowOffset;
@@ -82,6 +88,7 @@ void UCommonTextStyle::ApplyToTextBlock(const TSharedRef<STextBlock>& TextBlock)
 	TextBlock->SetStrikeBrush(&StrikeBrush);
 	TextBlock->SetMargin(Margin);
 	TextBlock->SetLineHeightPercentage(LineHeightPercentage);
+	TextBlock->SetApplyLineHeightToBottomLine(ApplyLineHeightToBottomLine);
 	TextBlock->SetColorAndOpacity(Color);
 	if (bUsesDropShadow)
 	{
@@ -212,6 +219,9 @@ void UCommonTextBlock::PostLoad()
 			// LineHeightPercentage
 			if (LineHeightPercentage == CDO->LineHeightPercentage) { LineHeightPercentage = DefaultStyleCDO->LineHeightPercentage; } else { bAllDefaults = false; }
 
+			// ApplyLineHeightToBottomLine
+			if (ApplyLineHeightToBottomLine == CDO->ApplyLineHeightToBottomLine) { ApplyLineHeightToBottomLine = DefaultStyleCDO->ApplyLineHeightToBottomLine; } else { bAllDefaults = false; }
+
 			// ColorAndOpacity
 			if (GetColorAndOpacity() == CDO->GetColorAndOpacity()) { SetColorAndOpacity(DefaultStyleCDO->Color); } else { bAllDefaults = false; }
 
@@ -297,6 +307,7 @@ void UCommonTextBlock::Serialize(FArchive& Ar)
 		FSlateBrush TempStrikeBrush;
 		FMargin TempMargin;
 		float TempLineHeightPercentage = 1.f;
+		bool TempApplyLineHeightToBottomLine = true;
 		FSlateColor TempColorAndOpacity;
 		FVector2D TempShadowOffset = FVector2D::ZeroVector;
 		FLinearColor TempShadowColorAndOpacity = FLinearColor::Transparent;
@@ -306,6 +317,7 @@ void UCommonTextBlock::Serialize(FArchive& Ar)
 		Swap(StrikeBrush, TempStrikeBrush);
 		Swap(Margin, TempMargin);
 		Swap(LineHeightPercentage, TempLineHeightPercentage);
+		Swap(ApplyLineHeightToBottomLine, TempApplyLineHeightToBottomLine);
 		Swap(ColorAndOpacity, TempColorAndOpacity);
 		Swap(ShadowOffset, TempShadowOffset);
 		Swap(ShadowColorAndOpacity, TempShadowColorAndOpacity);
@@ -316,6 +328,7 @@ void UCommonTextBlock::Serialize(FArchive& Ar)
 		Swap(TempStrikeBrush, StrikeBrush);
 		Swap(TempMargin, Margin);
 		Swap(TempLineHeightPercentage, LineHeightPercentage);
+		Swap(TempApplyLineHeightToBottomLine, ApplyLineHeightToBottomLine);
 		Swap(TempColorAndOpacity, ColorAndOpacity);
 		Swap(TempShadowOffset, ShadowOffset);
 		Swap(TempShadowColorAndOpacity, ShadowColorAndOpacity);
@@ -351,6 +364,11 @@ void UCommonTextBlock::SetLineHeightPercentage(float InLineHeightPercentage)
 	UTextLayoutWidget::SetLineHeightPercentage(InLineHeightPercentage);
 }
 
+void UCommonTextBlock::SetApplyLineHeightToBottomLine(bool InApplyLineHeightToBottomLine)
+{
+	UTextLayoutWidget::SetApplyLineHeightToBottomLine(InApplyLineHeightToBottomLine);
+}
+
 void UCommonTextBlock::SetStyle(TSubclassOf<UCommonTextStyle> InStyle)
 {
 	Style = InStyle;
@@ -474,6 +492,7 @@ void UCommonTextBlock::UpdateFromStyle()
 		SetStrikeBrush(TextStyle->StrikeBrush);
 		Margin = TextStyle->Margin;
 		LineHeightPercentage = TextStyle->LineHeightPercentage;
+		ApplyLineHeightToBottomLine = TextStyle->ApplyLineHeightToBottomLine;
 		SetColorAndOpacity(TextStyle->Color);
 
 		if (TextStyle->bUsesDropShadow)
@@ -528,6 +547,7 @@ PRAGMA_DISABLE_DEPRECATION_WARNINGS
 				GET_MEMBER_NAME_CHECKED(UCommonTextBlock, StrikeBrush),
 				GET_MEMBER_NAME_CHECKED(UCommonTextBlock, Margin),
 				GET_MEMBER_NAME_CHECKED(UCommonTextBlock, LineHeightPercentage),
+				GET_MEMBER_NAME_CHECKED(UCommonTextBlock, ApplyLineHeightToBottomLine),
 				GET_MEMBER_NAME_CHECKED(UCommonTextBlock, ColorAndOpacity),
 				GET_MEMBER_NAME_CHECKED(UCommonTextBlock, ShadowOffset),
 				GET_MEMBER_NAME_CHECKED(UCommonTextBlock, ShadowColorAndOpacity)
diff --git a/Engine/Plugins/Runtime/CommonUI/Source/CommonUI/Public/CommonTextBlock.h b/Engine/Plugins/Runtime/CommonUI/Source/CommonUI/Public/CommonTextBlock.h
--- a/Engine/Plugins/Runtime/CommonUI/Source/CommonUI/Public/CommonTextBlock.h
+++ b/Engine/Plugins/Runtime/CommonUI/Source/CommonUI/Public/CommonTextBlock.h
@@ -56,6 +56,10 @@ public:
 	UPROPERTY(EditDefaultsOnly, BlueprintReadOnly, Category = "Properties")
 	float LineHeightPercentage;
 
+	/** Whether to leave extra space below the last line due to line height */
+	UPROPERTY(EditDefaultsOnly, BlueprintReadOnly, Category = "Properties")
+	bool ApplyLineHeightToBottomLine;
+
 	UFUNCTION(BlueprintCallable, Category = "Common Text Style|Getters")
 	void GetFont(FSlateFontInfo& OutFont) const;
 
@@ -68,6 +72,9 @@ public:
 	UFUNCTION(BlueprintCallable, Category = "Common Text Style|Getters")
 	float GetLineHeightPercentage() const;
 
+	UFUNCTION(BlueprintCallable, Category = "Common Text Style|Getters")
+	bool GetApplyLineHeightToBottomLine() const;
+
 	UFUNCTION(BlueprintCallable, Category = "Common Text Style|Getters")
 	void GetShadowOffset(FVector2D& OutShadowOffset) const;
 
@@ -135,6 +142,9 @@ public:
 	UFUNCTION(BlueprintCallable, Category = "Common Text")
 	void SetLineHeightPercentage(float InLineHeightPercentage);
 
+	UFUNCTION(BlueprintCallable, Category = "Common Text")
+	void SetApplyLineHeightToBottomLine(bool InApplyLineHeightToBottomLine);
+
 	UFUNCTION(BlueprintCallable, Category = "Common Text")
 	void SetStyle(TSubclassOf<UCommonTextStyle> InStyle);
 
diff --git a/Engine/Plugins/VirtualProduction/RemoteControl/Source/RemoteControlUI/Private/UI/Drawers/SRCPanelDrawer.cpp b/Engine/Plugins/VirtualProduction/RemoteControl/Source/RemoteControlUI/Private/UI/Drawers/SRCPanelDrawer.cpp
--- a/Engine/Plugins/VirtualProduction/RemoteControl/Source/RemoteControlUI/Private/UI/Drawers/SRCPanelDrawer.cpp
+++ b/Engine/Plugins/VirtualProduction/RemoteControl/Source/RemoteControlUI/Private/UI/Drawers/SRCPanelDrawer.cpp
@@ -116,6 +116,7 @@ public:
 				ETextTransformPolicy::None,
 				FMargin(),
 				1.f,
+				true,
 				ETextJustify::Left),
 			LayoutScaleMultiplier, TextStyle);
 		return FVector2D(DesiredHorizontalTextSize.Y, DesiredHorizontalTextSize.X);
diff --git a/Engine/Source/Runtime/Slate/Private/Framework/Docking/STabSidebar.cpp b/Engine/Source/Runtime/Slate/Private/Framework/Docking/STabSidebar.cpp
--- a/Engine/Source/Runtime/Slate/Private/Framework/Docking/STabSidebar.cpp
+++ b/Engine/Source/Runtime/Slate/Private/Framework/Docking/STabSidebar.cpp
@@ -109,6 +109,7 @@ public:
 				ETextTransformPolicy::None,
 				FMargin(),
 				1.0f,
+				true,
 				ETextJustify::Left),
 			LayoutScaleMultiplier, TextStyle);
 		return FVector2D(DesiredHorizontalTextSize.Y, DesiredHorizontalTextSize.X);
diff --git a/Engine/Source/Runtime/Slate/Private/Framework/Text/TextLayout.cpp b/Engine/Source/Runtime/Slate/Private/Framework/Text/TextLayout.cpp
--- a/Engine/Source/Runtime/Slate/Private/Framework/Text/TextLayout.cpp
+++ b/Engine/Source/Runtime/Slate/Private/Framework/Text/TextLayout.cpp
@@ -365,6 +365,7 @@ void FTextLayout::CreateLineViewBlocks( int32 LineModelIndex, const int32 StopIn
 	}
 
 	FVector2D LineSize( ForceInitToZero );
+	float UnscaleLineHeight = 0.f;
 	
 	// Use a negative scroll offset since positive scrolling moves things negatively in screen space
 	FVector2D CurrentOffset(-ScrollOffset.X, TextLayoutSize.Height - ScrollOffset.Y);
@@ -425,7 +426,7 @@ void FTextLayout::CreateLineViewBlocks( int32 LineModelIndex, const int32 StopIn
 			CurrentHorizontalPos += Block->GetSize().X;
 		}
 
-		const float UnscaleLineHeight = MaxAboveBaseline + MaxBelowBaseline;
+		UnscaleLineHeight = MaxAboveBaseline + MaxBelowBaseline;
 
 		LineSize.X = CurrentHorizontalPos;
 		LineSize.Y = UnscaleLineHeight * LineHeightPercentage;
@@ -454,6 +455,7 @@ void FTextLayout::CreateLineViewBlocks( int32 LineModelIndex, const int32 StopIn
 	TextLayoutSize.DrawWidth = FMath::Max( TextLayoutSize.DrawWidth, LineSize.X ); // DrawWidth is the size of the longest line + the Margin
 	TextLayoutSize.WrappedWidth = FMath::Max( TextLayoutSize.WrappedWidth, (StopIndex == INDEX_NONE) ? LineSize.X : WrappedLineWidth ); // WrappedWidth is the size of the longest line + the Margin + any trailing whitespace width
 	TextLayoutSize.Height += LineSize.Y; // Height is the total height of all lines
+	OverHeight = LineSize.Y - UnscaleLineHeight;
 }
 
 void FTextLayout::JustifyLayout()
@@ -771,6 +773,10 @@ void FTextLayout::FlowLineLayout(const int32 LineModelIndex, const float Wrappin
 			}
 		}
 	}
+	if (!ApplyLineHeightToBottomLine)
+	{
+		TextLayoutSize.Height -= OverHeight;
+	}
 }
 
 void FTextLayout::FlowHighlights()
@@ -989,6 +995,7 @@ void FTextLayout::BeginLineLayout(FLineModel& LineModel)
 void FTextLayout::ClearView()
 {
 	TextLayoutSize = FTextLayoutSize();
+	OverHeight = 0.f;
 	LineViews.Empty();
 	LineViewsToJustify.Empty();
 }
@@ -1162,7 +1169,9 @@ FTextLayout::FTextLayout()
 	, Margin()
 	, Justification( ETextJustify::Left )
 	, LineHeightPercentage( 1.0f )
+	, ApplyLineHeightToBottomLine( true )
 	, TextLayoutSize()
+	, OverHeight( 0.0f )
 	, ViewSize( ForceInitToZero )
 	, ScrollOffset( ForceInitToZero )
 	, LineBreakIterator() // Initialized in FTextLayout::CreateWrappingCache if no custom iterator is provided
@@ -2080,6 +2089,20 @@ bool FTextLayout::RemoveLine(int32 LineIndex)
 				{
 					OffsetAdjustment += (LineViews[ViewIndex + 1].Offset.Y - LineView.Offset.Y);
 				}
+				else // Last line is being removed - OverHeight needs to be re-computed from the preceding line
+				{
+					float NewOverHeight = 0.f;
+					if (ViewIndex > 0)
+					{
+						FLineView& PrevLineView = LineViews[ViewIndex - 1];
+						NewOverHeight = PrevLineView.Size.Y - PrevLineView.TextHeight;
+					}
+					if (!ApplyLineHeightToBottomLine)
+					{
+						HeightAdjustment += NewOverHeight - OverHeight;
+					}
+					OverHeight = NewOverHeight;
+				}
 
 				LineViews.RemoveAt(ViewIndex);
 				LineViewsToJustify.Remove(ViewIndex);
@@ -2603,6 +2626,15 @@ void FTextLayout::SetLineHeightPercentage( float Value )
 	}
 }
 
+void FTextLayout::SetApplyLineHeightToBottomLine( bool Value )
+{
+	if ( ApplyLineHeightToBottomLine != Value )
+	{
+		ApplyLineHeightToBottomLine = Value; 
+		DirtyFlags |= ETextLayoutDirtyState::Layout;
+	}
+}
+
 void FTextLayout::SetWrappingWidth( float Value )
 {
 	const bool WasWrapping = WrappingWidth > 0.0f;
diff --git a/Engine/Source/Runtime/Slate/Private/Widgets/Input/SMultiLineEditableTextBox.cpp b/Engine/Source/Runtime/Slate/Private/Widgets/Input/SMultiLineEditableTextBox.cpp
--- a/Engine/Source/Runtime/Slate/Private/Widgets/Input/SMultiLineEditableTextBox.cpp
+++ b/Engine/Source/Runtime/Slate/Private/Widgets/Input/SMultiLineEditableTextBox.cpp
@@ -382,6 +382,11 @@ void SMultiLineEditableTextBox::SetLineHeightPercentage(const TAttribute<float>&
 	EditableText->SetLineHeightPercentage(InLineHeightPercentage);
 }
 
+void SMultiLineEditableTextBox::SetApplyLineHeightToBottomLine(const TAttribute<bool>& InApplyLineHeightToBottomLine)
+{
+	EditableText->SetApplyLineHeightToBottomLine(InApplyLineHeightToBottomLine);
+}
+
 void SMultiLineEditableTextBox::SetMargin(const TAttribute<FMargin>& InMargin)
 {
 	EditableText->SetMargin(InMargin);
diff --git a/Engine/Source/Runtime/Slate/Private/Widgets/Text/SMultiLineEditableText.cpp b/Engine/Source/Runtime/Slate/Private/Widgets/Text/SMultiLineEditableText.cpp
--- a/Engine/Source/Runtime/Slate/Private/Widgets/Text/SMultiLineEditableText.cpp
+++ b/Engine/Source/Runtime/Slate/Private/Widgets/Text/SMultiLineEditableText.cpp
@@ -211,6 +211,11 @@ void SMultiLineEditableText::SetLineHeightPercentage(const TAttribute<float>& In
 	EditableTextLayout->SetLineHeightPercentage(InLineHeightPercentage);
 }
 
+void SMultiLineEditableText::SetApplyLineHeightToBottomLine(const TAttribute<bool>& InApplyLineHeightToBottomLine)
+{
+	EditableTextLayout->SetApplyLineHeightToBottomLine(InApplyLineHeightToBottomLine);
+}
+
 void SMultiLineEditableText::SetMargin(const TAttribute<FMargin>& InMargin)
 {
 	EditableTextLayout->SetMargin(InMargin);
diff --git a/Engine/Source/Runtime/Slate/Private/Widgets/Text/SRichTextBlock.cpp b/Engine/Source/Runtime/Slate/Private/Widgets/Text/SRichTextBlock.cpp
--- a/Engine/Source/Runtime/Slate/Private/Widgets/Text/SRichTextBlock.cpp
+++ b/Engine/Source/Runtime/Slate/Private/Widgets/Text/SRichTextBlock.cpp
@@ -30,6 +30,7 @@ void SRichTextBlock::Construct( const FArguments& InArgs )
 	TransformPolicy = InArgs._TransformPolicy;
 	Margin = InArgs._Margin;
 	LineHeightPercentage = InArgs._LineHeightPercentage;
+	ApplyLineHeightToBottomLine = InArgs._ApplyLineHeightToBottomLine;
 	Justification = InArgs._Justification;
 	MinDesiredWidth = InArgs._MinDesiredWidth;
 
@@ -86,7 +87,7 @@ FVector2D SRichTextBlock::ComputeDesiredSize(float LayoutScaleMultiplier) const
 {
 	// ComputeDesiredSize will also update the text layout cache if required
 	const FVector2D TextSize = TextLayoutCache->ComputeDesiredSize(
-		FSlateTextBlockLayout::FWidgetDesiredSizeArgs(BoundText.Get(), HighlightText.Get(), WrapTextAt.Get(), AutoWrapText.Get(), WrappingPolicy.Get(), TransformPolicy.Get(), Margin.Get(), LineHeightPercentage.Get(), Justification.Get()),
+		FSlateTextBlockLayout::FWidgetDesiredSizeArgs(BoundText.Get(), HighlightText.Get(), WrapTextAt.Get(), AutoWrapText.Get(), WrappingPolicy.Get(), TransformPolicy.Get(), Margin.Get(), LineHeightPercentage.Get(), ApplyLineHeightToBottomLine.Get(), Justification.Get()),
 		LayoutScaleMultiplier * TextBlockScale, TextStyle) * TextBlockScale;
 
 	return FVector2D(FMath::Max(TextSize.X, MinDesiredWidth.Get()), TextSize.Y);
@@ -163,6 +164,11 @@ void SRichTextBlock::SetLineHeightPercentage(const TAttribute<float>& InLineHeig
 	SetAttribute(LineHeightPercentage, InLineHeightPercentage, EInvalidateWidgetReason::Layout);
 }
 
+void SRichTextBlock::SetApplyLineHeightToBottomLine(const TAttribute<bool>& InApplyLineHeightToBottomLine)
+{
+	SetAttribute(ApplyLineHeightToBottomLine, InApplyLineHeightToBottomLine, EInvalidateWidgetReason::Layout);
+}
+
 void SRichTextBlock::SetMargin(const TAttribute<FMargin>& InMargin)
 {
 	SetAttribute(Margin, InMargin, EInvalidateWidgetReason::Layout);
diff --git a/Engine/Source/Runtime/Slate/Private/Widgets/Text/STextBlock.cpp b/Engine/Source/Runtime/Slate/Private/Widgets/Text/STextBlock.cpp
--- a/Engine/Source/Runtime/Slate/Private/Widgets/Text/STextBlock.cpp
+++ b/Engine/Source/Runtime/Slate/Private/Widgets/Text/STextBlock.cpp
@@ -76,6 +76,8 @@ void STextBlock::PrivateRegisterAttributes(FSlateAttributeInitializer& Attribute
 		.OnValueChanged(FSlateAttributeDescriptor::FAttributeValueChangedDelegate::CreateStatic(FInvalidation::UpdateDesiredSize));
 	SLATE_ADD_MEMBER_ATTRIBUTE_DEFINITION(AttributeInitializer, LineHeightPercentage, FInvalidation::GetInvalidationNoneIfSimpleTextMode)
 		.OnValueChanged(FSlateAttributeDescriptor::FAttributeValueChangedDelegate::CreateStatic(FInvalidation::UpdateDesiredSize));
+	SLATE_ADD_MEMBER_ATTRIBUTE_DEFINITION(AttributeInitializer, ApplyLineHeightToBottomLine, FInvalidation::GetInvalidationNoneIfSimpleTextMode)
+		.OnValueChanged(FSlateAttributeDescriptor::FAttributeValueChangedDelegate::CreateStatic(FInvalidation::UpdateDesiredSize));
 	SLATE_ADD_MEMBER_ATTRIBUTE_DEFINITION(AttributeInitializer, MinDesiredWidth, EInvalidateWidgetReason::Layout);
 }
 
@@ -96,6 +98,7 @@ STextBlock::STextBlock()
 	, Margin(*this)
 	, Justification(*this)
 	, LineHeightPercentage(*this, 1.0f)
+	, ApplyLineHeightToBottomLine(*this, true)
 	, MinDesiredWidth(*this, 0.0f)
 	, Union_Flags(0)
 	, bSimpleTextMode(false)
@@ -129,6 +132,7 @@ void STextBlock::Construct( const FArguments& InArgs )
 
 	SetMargin(InArgs._Margin);
 	SetLineHeightPercentage(InArgs._LineHeightPercentage);
+	SetApplyLineHeightToBottomLine(InArgs._ApplyLineHeightToBottomLine);
 	SetJustification(InArgs._Justification);
 	SetMinDesiredWidth(InArgs._MinDesiredWidth);
 
@@ -459,6 +463,11 @@ void STextBlock::SetLineHeightPercentage(TAttribute<float> InLineHeightPercentag
 	LineHeightPercentage.Assign(*this, MoveTemp(InLineHeightPercentage));
 }
 
+void STextBlock::SetApplyLineHeightToBottomLine(TAttribute<bool> InApplyLineHeightToBottomLine)
+{
+	ApplyLineHeightToBottomLine.Assign(*this, MoveTemp(InApplyLineHeightToBottomLine));
+}
+
 void STextBlock::SetMargin(TAttribute<FMargin> InMargin)
 {
 	Margin.Assign(*this, MoveTemp(InMargin));
@@ -507,6 +516,7 @@ void STextBlock::UpdateTextBlockLayout(float LayoutScaleMultiplier) const
 			GetTransformPolicyImpl(),
 			Margin.Get(),
 			LineHeightPercentage.Get(),
+			ApplyLineHeightToBottomLine.Get(),
 			Justification.Get()
 		);
 		TextLayoutCache->ComputeDesiredSize(DesiredSizeArgs, LayoutScaleMultiplier);
diff --git a/Engine/Source/Runtime/Slate/Private/Widgets/Text/SlateEditableTextLayout.cpp b/Engine/Source/Runtime/Slate/Private/Widgets/Text/SlateEditableTextLayout.cpp
--- a/Engine/Source/Runtime/Slate/Private/Widgets/Text/SlateEditableTextLayout.cpp
+++ b/Engine/Source/Runtime/Slate/Private/Widgets/Text/SlateEditableTextLayout.cpp
@@ -518,6 +518,13 @@ void FSlateEditableTextLayout::SetLineHeightPercentage(const TAttribute<float>&
 	OwnerWidget->GetSlateWidget()->Invalidate(EInvalidateWidget::LayoutAndVolatility);
 }
 
+void FSlateEditableTextLayout::SetApplyLineHeightToBottomLine(const TAttribute<bool>& InApplyLineHeightToBottomLine)
+{
+	ApplyLineHeightToBottomLine = InApplyLineHeightToBottomLine;
+
+	OwnerWidget->GetSlateWidget()->Invalidate(EInvalidateWidget::LayoutAndVolatility);
+}
+
 void FSlateEditableTextLayout::SetOverflowPolicy(TOptional<ETextOverflowPolicy> InOverflowPolicy)
 {
 	if(OverflowPolicyOverride != InOverflowPolicy)
@@ -3589,6 +3596,7 @@ void FSlateEditableTextLayout::CacheDesiredSize(float LayoutScaleMultiplier)
 	TextLayout->SetWrappingPolicy(WrappingPolicy.Get());
 	TextLayout->SetMargin(MarginValue);
 	TextLayout->SetLineHeightPercentage(LineHeightPercentage.Get());
+	TextLayout->SetApplyLineHeightToBottomLine(ApplyLineHeightToBottomLine.Get());
 	TextLayout->SetJustification(Justification.Get());
 	TextLayout->SetVisibleRegion(FVector2D(CachedSize), FVector2D(ScrollOffset) * TextLayout->GetScale());
 	TextLayout->UpdateIfNeeded();
@@ -3614,7 +3622,7 @@ FVector2D FSlateEditableTextLayout::ComputeDesiredSize(float LayoutScaleMultipli
 		MarginValue.Right += CaretWidth;
 
 		const FVector2D HintTextSize = HintTextLayout->ComputeDesiredSize(
-			FSlateTextBlockLayout::FWidgetDesiredSizeArgs(HintText.Get(), FText::GetEmpty(), WrapTextAt.Get(), AutoWrapText.Get(), WrappingPolicy.Get(), ETextTransformPolicy::None, MarginValue, LineHeightPercentage.Get(), Justification.Get()),
+			FSlateTextBlockLayout::FWidgetDesiredSizeArgs(HintText.Get(), FText::GetEmpty(), WrapTextAt.Get(), AutoWrapText.Get(), WrappingPolicy.Get(), ETextTransformPolicy::None, MarginValue, LineHeightPercentage.Get(), ApplyLineHeightToBottomLine.Get(), Justification.Get()),
 			LayoutScaleMultiplier, HintTextStyle
 			);
 
diff --git a/Engine/Source/Runtime/Slate/Private/Widgets/Text/SlateTextBlockLayout.cpp b/Engine/Source/Runtime/Slate/Private/Widgets/Text/SlateTextBlockLayout.cpp
--- a/Engine/Source/Runtime/Slate/Private/Widgets/Text/SlateTextBlockLayout.cpp
+++ b/Engine/Source/Runtime/Slate/Private/Widgets/Text/SlateTextBlockLayout.cpp
@@ -86,6 +86,7 @@ FVector2D FSlateTextBlockLayout::ComputeDesiredSize(const FWidgetDesiredSizeArgs
 	TextLayout->SetMargin(InWidgetArgs.Margin);
 	TextLayout->SetJustification(InWidgetArgs.Justification);
 	TextLayout->SetLineHeightPercentage(InWidgetArgs.LineHeightPercentage);
+	TextLayout->SetApplyLineHeightToBottomLine(InWidgetArgs.ApplyLineHeightToBottomLine);
 
 	// Has the transform policy changed? If so we need a full refresh as that is destructive to the model text
 	if (PreviousTransformPolicy != TextLayout->GetTransformPolicy())
@@ -159,6 +160,7 @@ FVector2D FSlateTextBlockLayout::ComputeDesiredSize(const FWidgetArgs& InWidgetA
 	TextLayout->SetMargin(InWidgetArgs.Margin.Get());
 	TextLayout->SetJustification(InWidgetArgs.Justification.Get());
 	TextLayout->SetLineHeightPercentage(InWidgetArgs.LineHeightPercentage.Get());
+	TextLayout->SetApplyLineHeightToBottomLine(true);
 
 	// Has the transform policy changed? If so we need a full refresh as that is destructive to the model text
 	if (PreviousTransformPolicy != TextLayout->GetTransformPolicy())
diff --git a/Engine/Source/Runtime/Slate/Public/Framework/Text/TextLayout.h b/Engine/Source/Runtime/Slate/Public/Framework/Text/TextLayout.h
--- a/Engine/Source/Runtime/Slate/Public/Framework/Text/TextLayout.h
+++ b/Engine/Source/Runtime/Slate/Public/Framework/Text/TextLayout.h
@@ -415,6 +415,9 @@ public:
 	FORCEINLINE float GetLineHeightPercentage() const { return LineHeightPercentage; }
 	SLATE_API void SetLineHeightPercentage( float Value );
 
+	FORCEINLINE bool GetApplyLineHeightToBottomLine() const { return ApplyLineHeightToBottomLine; }
+	SLATE_API void SetApplyLineHeightToBottomLine( bool Value );
+
 	FORCEINLINE ETextJustify::Type GetJustification() const { return Justification; }
 	SLATE_API void SetJustification( ETextJustify::Type Value );
 
@@ -748,9 +751,15 @@ protected:
 	/** The percentage to modify a line height by. */
 	float LineHeightPercentage;
 
+	/** Whether or not line height should be applied to the last line. */
+	bool ApplyLineHeightToBottomLine;
+
 	/** The final size of the text layout on screen. */
 	FTextLayoutSize TextLayoutSize;
 
+	/** Extra height of the last line due to line height. */
+	float OverHeight;
+
 	/** The size of the text layout that can actually be seen from the parent widget */
 	FVector2D ViewSize;
 
diff --git a/Engine/Source/Runtime/Slate/Public/Widgets/Input/SMultiLineEditableTextBox.h b/Engine/Source/Runtime/Slate/Public/Widgets/Input/SMultiLineEditableTextBox.h
--- a/Engine/Source/Runtime/Slate/Public/Widgets/Input/SMultiLineEditableTextBox.h
+++ b/Engine/Source/Runtime/Slate/Public/Widgets/Input/SMultiLineEditableTextBox.h
@@ -359,6 +359,9 @@ public:
 	/** See LineHeightPercentage attribute */
 	SLATE_API void SetLineHeightPercentage(const TAttribute<float>& InLineHeightPercentage);
 
+	/** See ApplyLineHeightToBottomLine attribute */
+	SLATE_API void SetApplyLineHeightToBottomLine(const TAttribute<bool>& InApplyLineHeightToBottomLine);
+
 	/** See Margin attribute */
 	SLATE_API void SetMargin(const TAttribute<FMargin>& InMargin);
 
diff --git a/Engine/Source/Runtime/Slate/Public/Widgets/Text/SMultiLineEditableText.h b/Engine/Source/Runtime/Slate/Public/Widgets/Text/SMultiLineEditableText.h
--- a/Engine/Source/Runtime/Slate/Public/Widgets/Text/SMultiLineEditableText.h
+++ b/Engine/Source/Runtime/Slate/Public/Widgets/Text/SMultiLineEditableText.h
@@ -76,6 +76,7 @@ public:
 		, _Font()
 		, _Margin( FMargin() )
 		, _LineHeightPercentage( 1.0f )
+		, _ApplyLineHeightToBottomLine( true )
 		, _Justification( ETextJustify::Left )
 		, _IsReadOnly(false)
 		, _OnTextChanged()
@@ -135,6 +136,9 @@ public:
 		/** The amount to scale each lines height by. */
 		SLATE_ATTRIBUTE(float, LineHeightPercentage)
 
+		/** The amount to scale each lines height by. */
+		SLATE_ATTRIBUTE(bool, ApplyLineHeightToBottomLine)
+
 		/** How the text should be aligned with the margin. */
 		SLATE_ATTRIBUTE(ETextJustify::Type, Justification)
 
@@ -312,6 +316,9 @@ public:
 	/** See LineHeightPercentage attribute */
 	SLATE_API void SetLineHeightPercentage(const TAttribute<float>& InLineHeightPercentage);
 
+	/** See ApplyLineHeightToBottomLine attribute */
+	SLATE_API void SetApplyLineHeightToBottomLine(const TAttribute<bool>& InApplyLineHeightToBottomLine);
+
 	/** See Margin attribute */
 	SLATE_API void SetMargin(const TAttribute<FMargin>& InMargin);
 
diff --git a/Engine/Source/Runtime/Slate/Public/Widgets/Text/SRichTextBlock.h b/Engine/Source/Runtime/Slate/Public/Widgets/Text/SRichTextBlock.h
--- a/Engine/Source/Runtime/Slate/Public/Widgets/Text/SRichTextBlock.h
+++ b/Engine/Source/Runtime/Slate/Public/Widgets/Text/SRichTextBlock.h
@@ -46,6 +46,7 @@ public:
 		, _TextStyle( &FCoreStyle::Get().GetWidgetStyle<FTextBlockStyle>( "NormalText" ) )
 		, _Margin( FMargin() )
 		, _LineHeightPercentage( 1.0f )
+		, _ApplyLineHeightToBottomLine( true )
 		, _Justification( ETextJustify::Left )
 		, _TextShapingMethod()
 		, _TextFlowDirection()
@@ -94,6 +95,9 @@ public:
 		/** The amount to scale each lines height by. */
 		SLATE_ATTRIBUTE( float, LineHeightPercentage )
 
+		/** Whether to leave extra space below the last line due to line height. */
+		SLATE_ATTRIBUTE( bool, ApplyLineHeightToBottomLine )
+
 		/** How the text should be aligned with the margin. */
 		SLATE_ATTRIBUTE( ETextJustify::Type, Justification )
 
@@ -207,6 +211,9 @@ public:
 	/** See LineHeightPercentage attribute */
 	SLATE_API void SetLineHeightPercentage(const TAttribute<float>& InLineHeightPercentage);
 
+	/** See ApplyLineHeightToBottomLine attribute */
+	SLATE_API void SetApplyLineHeightToBottomLine(const TAttribute<bool>& InApplyLineHeightToBottomLine);
+
 	/** See Margin attribute */
 	SLATE_API void SetMargin(const TAttribute<FMargin>& InMargin);
 
@@ -277,6 +284,9 @@ private:
 	/** How the text should be aligned with the margin. */
 	TAttribute< float > LineHeightPercentage;
 
+	/** Whether to leave extra space below the last line due to line height. */
+	TAttribute< bool > ApplyLineHeightToBottomLine;
+
 	/** Prevents the text block from being smaller than desired in certain cases (e.g. when it is empty) */
 	TAttribute<float> MinDesiredWidth;
 
diff --git a/Engine/Source/Runtime/Slate/Public/Widgets/Text/STextBlock.h b/Engine/Source/Runtime/Slate/Public/Widgets/Text/STextBlock.h
--- a/Engine/Source/Runtime/Slate/Public/Widgets/Text/STextBlock.h
+++ b/Engine/Source/Runtime/Slate/Public/Widgets/Text/STextBlock.h
@@ -59,6 +59,7 @@ public:
 		, _TransformPolicy()
 		, _Margin()
 		, _LineHeightPercentage(1.0f)
+		, _ApplyLineHeightToBottomLine(true)
 		, _Justification(ETextJustify::Left)
 		, _MinDesiredWidth(0.0f)
 		, _TextShapingMethod()
@@ -120,6 +121,9 @@ public:
 		/** The amount to scale each lines height by. */
 		SLATE_ATTRIBUTE( float, LineHeightPercentage )
 
+		/** Whether to leave extra space below the last line due to line height. */
+		SLATE_ATTRIBUTE( bool, ApplyLineHeightToBottomLine )
+
 		/** How the text should be aligned with the margin. */
 		SLATE_ATTRIBUTE( ETextJustify::Type, Justification )
 
@@ -240,6 +244,9 @@ public:
 	/** See LineHeightPercentage attribute */
 	SLATE_API void SetLineHeightPercentage(TAttribute<float> InLineHeightPercentage);
 
+	/** See ApplyLineHeightToBottomLine attribute */
+	SLATE_API void SetApplyLineHeightToBottomLine(TAttribute<bool> InApplyLineHeightToBottomLine);
+
 	/** See Margin attribute */
 	SLATE_API void SetMargin(TAttribute<FMargin> InMargin);
 
@@ -360,6 +367,9 @@ private:
 	/** How the text should be aligned with the margin. */
 	TSlateAttribute<float> LineHeightPercentage;
 
+	/** Whether to leave extra space below the last line due to line height. */
+	TSlateAttribute<bool> ApplyLineHeightToBottomLine;
+
 	/** Prevents the text block from being smaller than desired in certain cases (e.g. when it is empty) */
 	TSlateAttribute<float> MinDesiredWidth;
 
diff --git a/Engine/Source/Runtime/Slate/Public/Widgets/Text/SlateEditableTextLayout.h b/Engine/Source/Runtime/Slate/Public/Widgets/Text/SlateEditableTextLayout.h
--- a/Engine/Source/Runtime/Slate/Public/Widgets/Text/SlateEditableTextLayout.h
+++ b/Engine/Source/Runtime/Slate/Public/Widgets/Text/SlateEditableTextLayout.h
@@ -117,6 +117,9 @@ public:
 	/** Set the amount to scale each lines height by */
 	SLATE_API void SetLineHeightPercentage(const TAttribute<float>& InLineHeightPercentage);
 
+	/** Set whether to leave extra space below the last line due to line height */
+	SLATE_API void SetApplyLineHeightToBottomLine(const TAttribute<bool>& InApplyLineHeightToBottomLine);
+
 	/** Set the text overflow policy that should be used to determine what happens to clipped text */
 	SLATE_API void SetOverflowPolicy(TOptional<ETextOverflowPolicy> InOverflowPolicy);
 
@@ -587,6 +590,9 @@ private:
 	/** The amount to scale each lines height by */
 	TAttribute<float> LineHeightPercentage;
 
+	/** Whether to leave extra space below the last line due to line height */
+	TAttribute<bool> ApplyLineHeightToBottomLine;
+
 	/** The information used to help identify who owns this text layout in the case of an error */
 	TAttribute<FString> DebugSourceInfo;
 
diff --git a/Engine/Source/Runtime/Slate/Public/Widgets/Text/SlateTextBlockLayout.h b/Engine/Source/Runtime/Slate/Public/Widgets/Text/SlateTextBlockLayout.h
--- a/Engine/Source/Runtime/Slate/Public/Widgets/Text/SlateTextBlockLayout.h
+++ b/Engine/Source/Runtime/Slate/Public/Widgets/Text/SlateTextBlockLayout.h
@@ -31,6 +31,7 @@ public:
 			const ETextTransformPolicy InTransformPolicy,
 			const FMargin& InMargin,
 			const float InLineHeightPercentage,
+			const bool InApplyLineHeightToBottomLine,
 			const ETextJustify::Type InJustification
 		)
 			: Text(InText)
@@ -38,6 +39,31 @@ public:
 			, Margin(InMargin)
 			, WrapTextAt(InWrapTextAt)
 			, LineHeightPercentage(InLineHeightPercentage)
+			, ApplyLineHeightToBottomLine(InApplyLineHeightToBottomLine)
+			, WrappingPolicy(InWrappingPolicy)
+			, TransformPolicy(InTransformPolicy)
+			, Justification(InJustification)
+			, AutoWrapText(InAutoWrapText)
+		{
+		}
+
+		FORCEINLINE FWidgetDesiredSizeArgs(
+			const FText& InText,
+			const FText& InHighlightText,
+			const float InWrapTextAt,
+			const bool InAutoWrapText,
+			const ETextWrappingPolicy InWrappingPolicy,
+			const ETextTransformPolicy InTransformPolicy,
+			const FMargin& InMargin,
+			const float InLineHeightPercentage,
+			const ETextJustify::Type InJustification
+		)
+			: Text(InText)
+			, HighlightText(InHighlightText)
+			, Margin(InMargin)
+			, WrapTextAt(InWrapTextAt)
+			, LineHeightPercentage(InLineHeightPercentage)
+			, ApplyLineHeightToBottomLine(true)
 			, WrappingPolicy(InWrappingPolicy)
 			, TransformPolicy(InTransformPolicy)
 			, Justification(InJustification)
@@ -50,6 +76,7 @@ public:
 		const FMargin Margin;
 		const float WrapTextAt = 0.f;
 		const float LineHeightPercentage;
+		const bool ApplyLineHeightToBottomLine;
 		const ETextWrappingPolicy WrappingPolicy;
 		const ETextTransformPolicy TransformPolicy;
 		const ETextJustify::Type Justification;
@@ -70,7 +97,7 @@ public:
 			const TAttribute<FMargin>& InMargin, 
 			const TAttribute<float>& InLineHeightPercentage, 
 			const TAttribute<ETextJustify::Type>& InJustification
-			)
+		)
 			: Text(InText)
 			, HighlightText(InHighlightText)
 			, WrapTextAt(InWrapTextAt)
diff --git a/Engine/Source/Runtime/UMG/Private/Components/MultiLineEditableText.cpp b/Engine/Source/Runtime/UMG/Private/Components/MultiLineEditableText.cpp
--- a/Engine/Source/Runtime/UMG/Private/Components/MultiLineEditableText.cpp
+++ b/Engine/Source/Runtime/UMG/Private/Components/MultiLineEditableText.cpp
@@ -133,6 +133,15 @@ void UMultiLineEditableText::OnLineHeightPercentageChanged(float InLineHeightPer
 	}
 }
 
+void UMultiLineEditableText::OnApplyLineHeightToBottomLineChanged(bool InApplyLineHeightToBottomLine)
+{
+	Super::OnApplyLineHeightToBottomLineChanged(InApplyLineHeightToBottomLine);
+	if (MyMultiLineEditableText.IsValid())
+	{
+		MyMultiLineEditableText->SetApplyLineHeightToBottomLine(InApplyLineHeightToBottomLine);
+	}
+}
+
 void UMultiLineEditableText::OnMarginChanged(const FMargin& InMargin)
 {
 	Super::OnMarginChanged(InMargin);
diff --git a/Engine/Source/Runtime/UMG/Private/Components/MultiLineEditableTextBox.cpp b/Engine/Source/Runtime/UMG/Private/Components/MultiLineEditableTextBox.cpp
--- a/Engine/Source/Runtime/UMG/Private/Components/MultiLineEditableTextBox.cpp
+++ b/Engine/Source/Runtime/UMG/Private/Components/MultiLineEditableTextBox.cpp
@@ -144,6 +144,15 @@ void UMultiLineEditableTextBox::OnLineHeightPercentageChanged(float InLineHeight
 	}
 }
 
+void UMultiLineEditableTextBox::OnApplyLineHeightToBottomLineChanged(bool InApplyLineHeightToBottomLine)
+{
+	Super::OnApplyLineHeightToBottomLineChanged(InApplyLineHeightToBottomLine);
+	if (MyEditableTextBlock.IsValid())
+	{
+		MyEditableTextBlock->SetApplyLineHeightToBottomLine(InApplyLineHeightToBottomLine);
+	}
+}
+
 void UMultiLineEditableTextBox::OnMarginChanged(const FMargin& InMargin)
 {
 	Super::OnMarginChanged(InMargin);
diff --git a/Engine/Source/Runtime/UMG/Private/Components/RichTextBlock.cpp b/Engine/Source/Runtime/UMG/Private/Components/RichTextBlock.cpp
--- a/Engine/Source/Runtime/UMG/Private/Components/RichTextBlock.cpp
+++ b/Engine/Source/Runtime/UMG/Private/Components/RichTextBlock.cpp
@@ -413,6 +413,15 @@ void URichTextBlock::OnLineHeightPercentageChanged(float InLineHeightPercentage)
 	}
 }
 
+void URichTextBlock::OnApplyLineHeightToBottomLineChanged(bool InApplyLineHeightToBottomLine)
+{
+	Super::OnApplyLineHeightToBottomLineChanged(InApplyLineHeightToBottomLine);
+	if (MyRichTextBlock.IsValid())
+	{
+		MyRichTextBlock->SetApplyLineHeightToBottomLine(InApplyLineHeightToBottomLine);
+	}
+}
+
 void URichTextBlock::OnMarginChanged(const FMargin& InMargin)
 {
 	Super::OnMarginChanged(InMargin);
diff --git a/Engine/Source/Runtime/UMG/Private/Components/TextBlock.cpp b/Engine/Source/Runtime/UMG/Private/Components/TextBlock.cpp
--- a/Engine/Source/Runtime/UMG/Private/Components/TextBlock.cpp
+++ b/Engine/Source/Runtime/UMG/Private/Components/TextBlock.cpp
@@ -161,6 +161,15 @@ void UTextBlock::OnLineHeightPercentageChanged(float InLineHeightPercentage)
 	}
 }
 
+void UTextBlock::OnApplyLineHeightToBottomLineChanged(bool InApplyLineHeightToBottomLine)
+{
+	Super::OnApplyLineHeightToBottomLineChanged(InApplyLineHeightToBottomLine);
+	if (MyTextBlock.IsValid())
+	{
+		MyTextBlock->SetApplyLineHeightToBottomLine(InApplyLineHeightToBottomLine);
+	}
+}
+
 void UTextBlock::OnMarginChanged(const FMargin& InMargin)
 {
 	Super::OnMarginChanged(InMargin);
diff --git a/Engine/Source/Runtime/UMG/Private/Components/TextWidgetTypes.cpp b/Engine/Source/Runtime/UMG/Private/Components/TextWidgetTypes.cpp
--- a/Engine/Source/Runtime/UMG/Private/Components/TextWidgetTypes.cpp
+++ b/Engine/Source/Runtime/UMG/Private/Components/TextWidgetTypes.cpp
@@ -23,6 +23,7 @@ UTextLayoutWidget::UTextLayoutWidget(const FObjectInitializer& ObjectInitializer
 	WrappingPolicy = ETextWrappingPolicy::DefaultWrapping;
 	Margin = FMargin(0.0f);
 	LineHeightPercentage = 1.0f;
+	ApplyLineHeightToBottomLine = true;
 }
 
 void UTextLayoutWidget::SetLineHeightPercentage(float InLineHeightPercentage)
@@ -34,6 +35,15 @@ void UTextLayoutWidget::SetLineHeightPercentage(float InLineHeightPercentage)
 	}
 }
 
+void UTextLayoutWidget::SetApplyLineHeightToBottomLine(bool InApplyLineHeightToBottomLine)
+{
+	if (ApplyLineHeightToBottomLine != InApplyLineHeightToBottomLine)
+	{
+		ApplyLineHeightToBottomLine = InApplyLineHeightToBottomLine;
+		OnApplyLineHeightToBottomLineChanged(InApplyLineHeightToBottomLine);
+	}
+}
+
 void UTextLayoutWidget::SetMargin(const FMargin& InMargin)
 {
 	if (Margin != InMargin)
diff --git a/Engine/Source/Runtime/UMG/Public/Components/MultiLineEditableText.h b/Engine/Source/Runtime/UMG/Public/Components/MultiLineEditableText.h
--- a/Engine/Source/Runtime/UMG/Public/Components/MultiLineEditableText.h
+++ b/Engine/Source/Runtime/UMG/Public/Components/MultiLineEditableText.h
@@ -184,6 +184,7 @@ public:
 protected:
 	//~ Begin UTextLayoutWidget Interface
 	UMG_API virtual void OnLineHeightPercentageChanged(float InLineHeightPercentage) override;
+	UMG_API virtual void OnApplyLineHeightToBottomLineChanged(bool InApplyLineHeightToBottomLine) override;
 	UMG_API virtual void OnMarginChanged(const FMargin& InMargin) override;
 	//~ End UTextLayoutWidget Interface
 
diff --git a/Engine/Source/Runtime/UMG/Public/Components/MultiLineEditableTextBox.h b/Engine/Source/Runtime/UMG/Public/Components/MultiLineEditableTextBox.h
--- a/Engine/Source/Runtime/UMG/Public/Components/MultiLineEditableTextBox.h
+++ b/Engine/Source/Runtime/UMG/Public/Components/MultiLineEditableTextBox.h
@@ -152,6 +152,7 @@ public:
 protected:
 	//~ Begin UTextLayoutWidget Interface
 	UMG_API virtual void OnLineHeightPercentageChanged(float InLineHeightPercentage) override;
+	UMG_API virtual void OnApplyLineHeightToBottomLineChanged(bool InApplyLineHeightToBottomLine) override;
 	UMG_API virtual void OnMarginChanged(const FMargin& InMargin) override;
 	//~ End UTextLayoutWidget Interface
 
diff --git a/Engine/Source/Runtime/UMG/Public/Components/RichTextBlock.h b/Engine/Source/Runtime/UMG/Public/Components/RichTextBlock.h
--- a/Engine/Source/Runtime/UMG/Public/Components/RichTextBlock.h
+++ b/Engine/Source/Runtime/UMG/Public/Components/RichTextBlock.h
@@ -229,6 +229,7 @@ public:
 protected:
 	//~ Begin UTextLayoutWidget Interface
 	UMG_API virtual void OnLineHeightPercentageChanged(float InLineHeightPercentage) override;
+	UMG_API virtual void OnApplyLineHeightToBottomLineChanged(bool InApplyLineHeightToBottomLine) override;
 	UMG_API virtual void OnMarginChanged(const FMargin& InMargin) override;
 	//~ End UTextLayoutWidget Interface
 
diff --git a/Engine/Source/Runtime/UMG/Public/Components/TextBlock.h b/Engine/Source/Runtime/UMG/Public/Components/TextBlock.h
--- a/Engine/Source/Runtime/UMG/Public/Components/TextBlock.h
+++ b/Engine/Source/Runtime/UMG/Public/Components/TextBlock.h
@@ -258,6 +258,7 @@ public:
 protected:
 	//~ Begin UTextLayoutWidget Interface
 	UMG_API virtual void OnLineHeightPercentageChanged(float InLineHeightPercentage) override;
+	UMG_API virtual void OnApplyLineHeightToBottomLineChanged(bool InApplyLineHeightToBottomLine) override;
 	UMG_API virtual void OnMarginChanged(const FMargin& InMargin) override;
 	//~ End UTextLayoutWidget Interface
 
diff --git a/Engine/Source/Runtime/UMG/Public/Components/TextWidgetTypes.h b/Engine/Source/Runtime/UMG/Public/Components/TextWidgetTypes.h
--- a/Engine/Source/Runtime/UMG/Public/Components/TextWidgetTypes.h
+++ b/Engine/Source/Runtime/UMG/Public/Components/TextWidgetTypes.h
@@ -66,10 +66,12 @@ public:
 	virtual void SetJustification(ETextJustify::Type InJustification) { Justification = InJustification; }
 
 	UMG_API void SetLineHeightPercentage(float InLineHeightPercentage);
+	UMG_API void SetApplyLineHeightToBottomLine(bool InApplyLineHeightToBottomLine);
 	UMG_API void SetMargin(const FMargin& InMargin);
 
 protected:
 	virtual void OnLineHeightPercentageChanged(float InLineHeightPercentage) {};
+	virtual void OnApplyLineHeightToBottomLineChanged(bool InApplyLineHeightToBottomLine) {};
 	virtual void OnMarginChanged(const FMargin& InMargin) {};
 
 	/** Synchronize the properties with the given widget. A template as the Slate widgets conform to the same API, but don't derive from a common base. */
@@ -84,6 +86,7 @@ protected:
 		InWidget.SetWrappingPolicy(WrappingPolicy);
 		InWidget.SetMargin(Margin);
 		InWidget.SetLineHeightPercentage(LineHeightPercentage);
+		InWidget.SetApplyLineHeightToBottomLine(ApplyLineHeightToBottomLine);
 	}
 
 	/** Controls how the text within this widget should be shaped. */
@@ -113,4 +116,8 @@ protected:
 	/** The amount to scale each lines height by. */
 	UPROPERTY(EditAnywhere, BlueprintReadWrite, Setter, Category=Appearance, AdvancedDisplay)
 	float LineHeightPercentage;
+
+	/** Whether to leave extra space below the last line due to line height. */
+	UPROPERTY(EditAnywhere, BlueprintReadWrite, Setter, Category=Appearance, AdvancedDisplay)
+	bool ApplyLineHeightToBottomLine;
 };
