Renaming UE_COMMONUI_PLATFORM_SUPPORTS_TOUCH macro to UE_COMMONUI_PLATFORM_REQUIRES_CURSOR_HIDDEN_FOR_TOUCH and making sure it is only enabled for the relevant platforms. Clarifying the name of the Macro and Property both more accurately describe it's function and clarify that it is different from the Macro and Property of the same name in CommonInput.

The macro was introduced to get rid of a platform check but the ended up applying the logic gated by that platform check to Android, iOS and Windows non-Shipping platforms when they didn't need it. Also, the naming of the macro made it seem like a required macro for all touch-supporting platforms when it was only required for certain platforms.

#rnx
#tests Tested with a preflight on android, switch and mobile streaming platforms

diff --git a/Engine/Plugins/Runtime/CommonUI/Source/CommonUI/CommonUI.Build.cs b/Engine/Plugins/Runtime/CommonUI/Source/CommonUI/CommonUI.Build.cs
--- a/Engine/Plugins/Runtime/CommonUI/Source/CommonUI/CommonUI.Build.cs
+++ b/Engine/Plugins/Runtime/CommonUI/Source/CommonUI/CommonUI.Build.cs
@@ -46,19 +46,9 @@ public class CommonUI : ModuleRules
             );
         }
 		PrivateDefinitions.Add("UE_COMMONUI_PLATFORM_KBM_REQUIRES_ATTACHED_MOUSE=" + (bPlatformKBMRequiresAttachedMouse ? "1" : "0"));
-		PrivateDefinitions.Add("UE_COMMONUI_PLATFORM_SUPPORTS_TOUCH=" + (bPlatformSupportsTouch ? "1" : "0"));
+		PrivateDefinitions.Add("UE_COMMONUI_PLATFORM_REQUIRES_CURSOR_HIDDEN_FOR_TOUCH=" + (bPlatformRequiresCursorHiddenForTouch ? "1" : "0"));
 	}
 
 	protected virtual bool bPlatformKBMRequiresAttachedMouse { get { return false; } }
-	protected virtual bool bPlatformSupportsTouch
-	{
-		get
-		{
-			// Support touch testing until touch is supported on desktop
-			return (Target.Platform.IsInGroup(UnrealPlatformGroup.Desktop) && Target.Configuration != UnrealTargetConfiguration.Shipping)
-					// Platforms always with touch support
-					|| (Target.Platform == UnrealTargetPlatform.Android)
-					|| (Target.Platform == UnrealTargetPlatform.IOS);
-		}
-	}
+	protected virtual bool bPlatformRequiresCursorHiddenForTouch { get { return false; } }
 }
diff --git a/Engine/Plugins/Runtime/CommonUI/Source/CommonUI/Private/Input/CommonAnalogCursor.cpp b/Engine/Plugins/Runtime/CommonUI/Source/CommonUI/Private/Input/CommonAnalogCursor.cpp
--- a/Engine/Plugins/Runtime/CommonUI/Source/CommonUI/Private/Input/CommonAnalogCursor.cpp
+++ b/Engine/Plugins/Runtime/CommonUI/Source/CommonUI/Private/Input/CommonAnalogCursor.cpp
@@ -325,7 +325,7 @@ bool FCommonAnalogCursor::HandleMouseButtonDownEvent(FSlateApplication& SlateApp
 {
 	if (FAnalogCursor::IsRelevantInput(PointerEvent))
 	{
-#if UE_COMMONUI_PLATFORM_SUPPORTS_TOUCH	
+#if UE_COMMONUI_PLATFORM_REQUIRES_CURSOR_HIDDEN_FOR_TOUCH	
 		// Some platforms don't register as switching its input type, so detect touch input here to hide the cursor.
 		if (PointerEvent.IsTouchEvent() && ShouldHideCursor())
 		{
