Removing the UE_COMMONINPUT_PLATFORM_SUPPORTS_TOUCH macro since the naming was not accurately describe it's function and it was also disabling touch support for certain platforms even though they had their touch support flag enabled in their DataDrivenPlatform.ini file. It was only meant to replace a specific platform check but was instead affecting the touch support behavior of multiple platforms.

The UE_COMMONINPUT_PLATFORM_SUPPORTS_TOUCH macro will be replaced with a new UE_COMMONINPUT_FORCE_TOUCH_SUPPORT_DISABLED macro and only enabled on the required platform to match the logic that was in place before the introduction of the UE_COMMONINPUT_PLATFORM_SUPPORTS_TOUCH macro.

Also adjusted the logic that overrides the touch support flag on Desktop builds so that the override only applies in the Editor (to support testing with UseMouseForTouch setting enabled or with URemote in the Editor). The previous logic was only overriding the touch support flag in non-Shipping Destop builds, which was causing different touch behaviors on Desktop Test builds and Desktop Shipping builds.

#rnx
#tests Tested touch support still works when using PIE with UseMouseForTouch enabled or URemote. Tested with a preflight on android, switch and mobile streaming platforms.

diff --git a/Engine/Plugins/Runtime/CommonUI/Source/CommonInput/CommonInput.Build.cs b/Engine/Plugins/Runtime/CommonUI/Source/CommonInput/CommonInput.Build.cs
--- a/Engine/Plugins/Runtime/CommonUI/Source/CommonInput/CommonInput.Build.cs
+++ b/Engine/Plugins/Runtime/CommonUI/Source/CommonInput/CommonInput.Build.cs
@@ -45,7 +45,7 @@ public class CommonInput : ModuleRules
 		}
 
 		PrivateDefinitions.Add("UE_COMMONINPUT_PLATFORM_KBM_REQUIRES_ATTACHED_MOUSE=" + (bPlatformKBMRequiresAttachedMouse ? "1" : "0"));
-		PrivateDefinitions.Add("UE_COMMONINPUT_PLATFORM_SUPPORTS_TOUCH=" + (bPlatformSupportsTouch ? "1" : "0"));
+		PrivateDefinitions.Add("UE_COMMONINPUT_FORCE_TOUCH_SUPPORT_DISABLED=" + (bForceTouchSupportDisabled ? "1" : "0"));
 	}
 
 	static public string ToCommonUIPlatform(UnrealTargetPlatform TargetPlatform)
@@ -75,15 +75,5 @@ public class CommonInput : ModuleRules
 	}
 
 	protected virtual bool bPlatformKBMRequiresAttachedMouse { get { return false; } }
-	protected virtual bool bPlatformSupportsTouch
-	{
-		get
-		{
-			// Support touch testing until touch is supported on desktop
-			return (Target.Platform.IsInGroup(UnrealPlatformGroup.Desktop) && Target.Configuration != UnrealTargetConfiguration.Shipping)
-					// Platforms always with touch support
-					|| (Target.Platform == UnrealTargetPlatform.Android)
-					|| (Target.Platform == UnrealTargetPlatform.IOS);
-		}
-	}
+	protected virtual bool bForceTouchSupportDisabled { get { return false; } }
 }
diff --git a/Engine/Plugins/Runtime/CommonUI/Source/CommonInput/Private/CommonInputSubsystem.cpp b/Engine/Plugins/Runtime/CommonUI/Source/CommonInput/Private/CommonInputSubsystem.cpp
--- a/Engine/Plugins/Runtime/CommonUI/Source/CommonInput/Private/CommonInputSubsystem.cpp
+++ b/Engine/Plugins/Runtime/CommonUI/Source/CommonInput/Private/CommonInputSubsystem.cpp
@@ -524,9 +524,10 @@ bool UCommonInputSubsystem::PlatformSupportsInputType(ECommonInputType InInputTy
 		break;
 		case ECommonInputType::Touch:
 		{
-			bPlatformSupportsInput &= UE_COMMONINPUT_PLATFORM_SUPPORTS_TOUCH != 0;
-#if PLATFORM_DESKTOP && !UE_BUILD_SHIPPING
-			bPlatformSupportsInput = true; // Enable touch in debug editor
+			bPlatformSupportsInput &= !UE_COMMONINPUT_FORCE_TOUCH_SUPPORT_DISABLED;
+#if WITH_EDITOR
+			// Support touch testing (testing with UseMouseForTouch setting enabled or with URemote in the editor) until touch is supported on desktop
+			bPlatformSupportsInput = true;
 #endif
 		}
 		break;
