PR #10814: Add method to clear an Activatable Widget's Focus Restoration Target.

diff --git a/Engine/Plugins/Runtime/CommonUI/Source/CommonUI/Private/CommonActivatableWidget.cpp b/Engine/Plugins/Runtime/CommonUI/Source/CommonUI/Private/CommonActivatableWidget.cpp
--- a/Engine/Plugins/Runtime/CommonUI/Source/CommonUI/Private/CommonActivatableWidget.cpp
+++ b/Engine/Plugins/Runtime/CommonUI/Source/CommonUI/Private/CommonActivatableWidget.cpp
@@ -77,6 +77,14 @@ UWidget* UCommonActivatableWidget::GetDesiredFocusTarget() const
 	return NativeGetDesiredFocusTarget();
 }
 
+void UCommonActivatableWidget::ClearFocusRestorationTarget()
+{
+	if (TSharedPtr<FActivatableTreeNode> PinnedTreeNode = InputTreeNode.Pin())
+	{
+		PinnedTreeNode->ClearFocusRestorationTarget();
+	}
+}
+
 TOptional<FActivationMetadata> UCommonActivatableWidget::GetActivationMetadata() const
 {
 	return TOptional<FActivationMetadata>();
diff --git a/Engine/Plugins/Runtime/CommonUI/Source/CommonUI/Private/Input/UIActionRouterTypes.h b/Engine/Plugins/Runtime/CommonUI/Source/CommonUI/Private/Input/UIActionRouterTypes.h
--- a/Engine/Plugins/Runtime/CommonUI/Source/CommonUI/Private/Input/UIActionRouterTypes.h
+++ b/Engine/Plugins/Runtime/CommonUI/Source/CommonUI/Private/Input/UIActionRouterTypes.h
@@ -114,6 +114,7 @@ public:
 	
 	FActivatableTreeNodeRef AddChildNode(UCommonActivatableWidget& InActivatableWidget);
 	void CacheFocusRestorationTarget();
+	void ClearFocusRestorationTarget() { FocusRestorationTarget.Reset(); }
 	TSharedPtr<SWidget> GetFocusFallbackTarget() const;
 
 	bool IsExclusiveParentOfWidget(const TSharedPtr<SWidget>& SlateWidget) const;
diff --git a/Engine/Plugins/Runtime/CommonUI/Source/CommonUI/Public/CommonActivatableWidget.h b/Engine/Plugins/Runtime/CommonUI/Source/CommonUI/Public/CommonActivatableWidget.h
--- a/Engine/Plugins/Runtime/CommonUI/Source/CommonUI/Public/CommonActivatableWidget.h
+++ b/Engine/Plugins/Runtime/CommonUI/Source/CommonUI/Public/CommonActivatableWidget.h
@@ -66,6 +66,10 @@ public:
 	UFUNCTION(BlueprintCallable, Category = ActivatableWidget)
 	UWidget* GetDesiredFocusTarget() const;
 
+	/** Clears the cached focus target that's set when bAutoRestoreFocus is true */
+	UFUNCTION(BlueprintCallable, Category = ActivatableWidget)
+	void ClearFocusRestorationTarget();
+
 	FSimpleMulticastDelegate& OnActivated() const { return OnActivatedEvent; }
 	FSimpleMulticastDelegate& OnDeactivated() const { return OnDeactivatedEvent; }
 
