PR #11183: [Common UI] Add the bClearFocusRestorationTargetOfDeactivatedWidgets setting to UCommonActivatableWidgetSwitcher

diff --git a/Engine/Plugins/Runtime/CommonUI/Source/CommonUI/Private/CommonActivatableWidget.cpp b/Engine/Plugins/Runtime/CommonUI/Source/CommonUI/Private/CommonActivatableWidget.cpp
--- a/Engine/Plugins/Runtime/CommonUI/Source/CommonUI/Private/CommonActivatableWidget.cpp
+++ b/Engine/Plugins/Runtime/CommonUI/Source/CommonUI/Private/CommonActivatableWidget.cpp
@@ -172,6 +172,11 @@ void UCommonActivatableWidget::InternalProcessDeactivation()
 	NativeOnDeactivated();
 }
 
+TWeakPtr<FActivatableTreeNode> UCommonActivatableWidget::GetInputTreeNode() const
+{
+	return InputTreeNode;
+}
+
 void UCommonActivatableWidget::RegisterInputTreeNode(const TSharedPtr<FActivatableTreeNode>& OwnerNode)
 {
 	InputTreeNode = OwnerNode;
diff --git a/Engine/Plugins/Runtime/CommonUI/Source/CommonUI/Private/CommonActivatableWidgetSwitcher.cpp b/Engine/Plugins/Runtime/CommonUI/Source/CommonUI/Private/CommonActivatableWidgetSwitcher.cpp
--- a/Engine/Plugins/Runtime/CommonUI/Source/CommonUI/Private/CommonActivatableWidgetSwitcher.cpp
+++ b/Engine/Plugins/Runtime/CommonUI/Source/CommonUI/Private/CommonActivatableWidgetSwitcher.cpp
@@ -68,6 +68,11 @@ void UCommonActivatableWidgetSwitcher::DeactivateActiveWidget()
 {
 	if (UCommonActivatableWidget* ActivatableWidget = Cast<UCommonActivatableWidget>(GetActiveWidget()))
 	{
+		if (bClearFocusRestorationTargetOfDeactivatedWidgets)
+		{
+			ActivatableWidget->ClearFocusRestorationTarget();
+		}
+
 		ActivatableWidget->DeactivateWidget();
 	}
 }
diff --git a/Engine/Plugins/Runtime/CommonUI/Source/CommonUI/Public/CommonActivatableWidget.h b/Engine/Plugins/Runtime/CommonUI/Source/CommonUI/Public/CommonActivatableWidget.h
--- a/Engine/Plugins/Runtime/CommonUI/Source/CommonUI/Public/CommonActivatableWidget.h
+++ b/Engine/Plugins/Runtime/CommonUI/Source/CommonUI/Public/CommonActivatableWidget.h
@@ -100,6 +100,7 @@ public:
 	bool SetsVisibilityOnActivated() const { return bSetVisibilityOnActivated; }
 	bool SetsVisibilityOnDeactivated() const { return bSetVisibilityOnDeactivated; }
 
+	TWeakPtr<FActivatableTreeNode> GetInputTreeNode() const;
 	void RegisterInputTreeNode(const TSharedPtr<FActivatableTreeNode>& OwnerNode);
 	void ClearActiveHoldInputs();
 
diff --git a/Engine/Plugins/Runtime/CommonUI/Source/CommonUI/Public/CommonActivatableWidgetSwitcher.h b/Engine/Plugins/Runtime/CommonUI/Source/CommonUI/Public/CommonActivatableWidgetSwitcher.h
--- a/Engine/Plugins/Runtime/CommonUI/Source/CommonUI/Public/CommonActivatableWidgetSwitcher.h
+++ b/Engine/Plugins/Runtime/CommonUI/Source/CommonUI/Public/CommonActivatableWidgetSwitcher.h
@@ -26,6 +26,10 @@ protected:
 	virtual void HandleOutgoingWidget() override;
 	virtual void HandleSlateActiveIndexChanged(int32 ActiveIndex) override;
 
+protected:
+	UPROPERTY(EditAnywhere, BlueprintReadOnly, Category = "Switcher")
+	bool bClearFocusRestorationTargetOfDeactivatedWidgets = false;
+
 private:
 	void HandleOwningWidgetActivationChanged(const bool bIsActivated);
 
