PR #11192: CommonUI: Add virtual method to CommonAnalogCusor to allow customizing when Left Click should be simulated by [REDACTED_PERSON] events

diff --git a/Engine/Plugins/Runtime/CommonUI/Source/CommonUI/Private/Input/CommonAnalogCursor.cpp b/Engine/Plugins/Runtime/CommonUI/Source/CommonUI/Private/Input/CommonAnalogCursor.cpp
--- a/Engine/Plugins/Runtime/CommonUI/Source/CommonUI/Private/Input/CommonAnalogCursor.cpp
+++ b/Engine/Plugins/Runtime/CommonUI/Source/CommonUI/Private/Input/CommonAnalogCursor.cpp
@@ -243,7 +243,7 @@ bool FCommonAnalogCursor::HandleKeyDownEvent(FSlateApplication& SlateApp, const
 		{
 			return true;
 		}
-		else
+		else if (!bIsVirtualAccept || ShouldVirtualAcceptSimulateMouseButton(InKeyEvent, IE_Pressed))
 		{
 			// There is no awareness on a mouse event of whether it's real or not, so mark that here.
 			UCommonInputSubsystem& InputSubsytem = ActionRouter.GetInputSubsystem();
@@ -275,7 +275,7 @@ bool FCommonAnalogCursor::HandleKeyUpEvent(FSlateApplication& SlateApp, const FK
 		{
 			return true;
 		}
-		else
+		else if (!bIsVirtualAccept || ShouldVirtualAcceptSimulateMouseButton(InKeyEvent, IE_Released))
 		{
 			return FAnalogCursor::HandleKeyUpEvent(SlateApp, InKeyEvent);
 		}
diff --git a/Engine/Plugins/Runtime/CommonUI/Source/CommonUI/Public/Input/CommonAnalogCursor.h b/Engine/Plugins/Runtime/CommonUI/Source/CommonUI/Public/Input/CommonAnalogCursor.h
--- a/Engine/Plugins/Runtime/CommonUI/Source/CommonUI/Public/Input/CommonAnalogCursor.h
+++ b/Engine/Plugins/Runtime/CommonUI/Source/CommonUI/Public/Input/CommonAnalogCursor.h
@@ -2,6 +2,7 @@
 
 #pragma once
 
+#include "Engine/EngineBaseTypes.h"
 #include "Framework/Application/AnalogCursor.h"
 #include "InputCoreTypes.h"
 #include "Rendering/SlateRenderTransform.h"
@@ -52,6 +53,8 @@ public:
 
 	virtual bool IsAnalogMovementEnabled() const { return bIsAnalogMovementEnabled; }
 
+	virtual bool ShouldVirtualAcceptSimulateMouseButton(const FKeyEvent& InKeyEvent, EInputEvent InputEvent) const { return true; }
+
 protected:
 	FCommonAnalogCursor(const UCommonUIActionRouterBase& InActionRouter);
 	virtual void Initialize();
