PR #11209: CommonUI: Make public UCommonUserWidget::(Un)RegisterScrollRecipient functions and expose them to blueprint
PR #11212: CommonUI: Fix ScrollRecipients not working if they were registered before the owning widget was constructed
PR #11153: [Common UI] Expose FUIActionBindingHandle to blueprints

diff --git a/Engine/Plugins/Runtime/CommonUI/Source/CommonUI/Private/CommonUserWidget.cpp b/Engine/Plugins/Runtime/CommonUI/Source/CommonUI/Private/CommonUserWidget.cpp
--- a/Engine/Plugins/Runtime/CommonUI/Source/CommonUI/Private/CommonUserWidget.cpp
+++ b/Engine/Plugins/Runtime/CommonUI/Source/CommonUI/Private/CommonUserWidget.cpp
@@ -140,6 +140,22 @@ void UCommonUserWidget::UnregisterScrollRecipient(const UWidget& AnalogScrollRec
 	}
 }
 
+void UCommonUserWidget::RegisterScrollRecipientExternal(const UWidget* AnalogScrollRecipient)
+{
+	if (AnalogScrollRecipient != nullptr)
+	{
+		RegisterScrollRecipient(*AnalogScrollRecipient);
+	}
+}
+
+void UCommonUserWidget::UnregisterScrollRecipientExternal(const UWidget* AnalogScrollRecipient)
+{
+	if (AnalogScrollRecipient != nullptr)
+	{
+		UnregisterScrollRecipient(*AnalogScrollRecipient);
+	}
+}
+
 void UCommonUserWidget::OnWidgetRebuilt()
 {
 	// Using OnWidgetRebuilt instead of NativeConstruct to ensure we register ourselves with the ActionRouter before the child receives NativeConstruct
@@ -164,6 +180,17 @@ void UCommonUserWidget::OnWidgetRebuilt()
 	}
 
 	Super::OnWidgetRebuilt();
+
+	if (UCommonUIActionRouterBase* ActionRouter = UCommonUIActionRouterBase::Get(*this))
+	{
+		for (const TWeakObjectPtr<const UWidget>& WidgetPtr : GetScrollRecipients())
+		{
+			if (const UWidget* Widget = WidgetPtr.Get())
+			{
+				ActionRouter->RegisterScrollRecipient(*Widget);
+			}
+		}
+	}
 }
 
 void UCommonUserWidget::NativeDestruct()
diff --git a/Engine/Plugins/Runtime/CommonUI/Source/CommonUI/Public/CommonUserWidget.h b/Engine/Plugins/Runtime/CommonUI/Source/CommonUI/Public/CommonUserWidget.h
--- a/Engine/Plugins/Runtime/CommonUI/Source/CommonUI/Public/CommonUserWidget.h
+++ b/Engine/Plugins/Runtime/CommonUI/Source/CommonUI/Public/CommonUserWidget.h
@@ -25,9 +25,28 @@ public:
 	UFUNCTION(BlueprintCallable, Category = CommonUserWidget)
 	void SetConsumePointerInput(bool bInConsumePointerInput);
 
+	/** Add a widget to the list of widgets to get scroll events for this input root node */
+	UFUNCTION(BlueprintCallable, Category = CommonUserWidget)
+	void RegisterScrollRecipientExternal(const UWidget* AnalogScrollRecipient);
+
+	/** Remove a widget from the list of widgets to get scroll events for this input root node */
+	UFUNCTION(BlueprintCallable, Category = CommonUserWidget)
+	void UnregisterScrollRecipientExternal(const UWidget* AnalogScrollRecipient);
+
+public:
+
 	const TArray<FUIActionBindingHandle>& GetActionBindings() const { return ActionBindings; }
 	const TArray<TWeakObjectPtr<const UWidget>> GetScrollRecipients() const { return ScrollRecipients; }
 
+	/**
+	 * Convenience methods for menu action registrations (any UWidget can register via FCommonUIActionRouter directly, though generally that shouldn't be needed).
+	 * Persistent bindings are *always* listening for input while registered, while normal bindings are only listening when all of this widget's activatable parents are activated.
+	 */
+	FUIActionBindingHandle RegisterUIActionBinding(const FBindUIActionArgs& BindActionArgs);
+
+	void RemoveActionBinding(FUIActionBindingHandle ActionBinding);
+	void AddActionBinding(FUIActionBindingHandle ActionBinding);
+
 protected:
 	virtual void OnWidgetRebuilt() override;
 	virtual void NativeDestruct() override;
@@ -64,15 +83,6 @@ protected:
 	void RegisterScrollRecipient(const UWidget& AnalogScrollRecipient);
 	void UnregisterScrollRecipient(const UWidget& AnalogScrollRecipient);
 
-	/** 
-	 * Convenience methods for menu action registrations (any UWidget can register via FCommonUIActionRouter directly, though generally that shouldn't be needed).
-	 * Persistent bindings are *always* listening for input while registered, while normal bindings are only listening when all of this widget's activatable parents are activated.
-	 */
-	FUIActionBindingHandle RegisterUIActionBinding(const FBindUIActionArgs& BindActionArgs);
-
-	void RemoveActionBinding(FUIActionBindingHandle ActionBinding);
-	void AddActionBinding(FUIActionBindingHandle ActionBinding);
-
 	/** True to generally display this widget's actions in the action bar, assuming it has actions. */
 	UPROPERTY(EditAnywhere, BlueprintReadOnly, Category = Input, meta = (AllowPrivateAccess = true))
 	bool bDisplayInActionBar = false;
diff --git a/Engine/Plugins/Runtime/CommonUI/Source/CommonUI/Public/Input/UIActionBindingHandle.h b/Engine/Plugins/Runtime/CommonUI/Source/CommonUI/Public/Input/UIActionBindingHandle.h
--- a/Engine/Plugins/Runtime/CommonUI/Source/CommonUI/Public/Input/UIActionBindingHandle.h
+++ b/Engine/Plugins/Runtime/CommonUI/Source/CommonUI/Public/Input/UIActionBindingHandle.h
@@ -13,8 +13,11 @@ struct FScriptContainerElement;
 class UWidget;
 enum class ECommonInputMode : uint8;
 
+USTRUCT(BlueprintType, DisplayName = "UI Action Binding Handle")
 struct COMMONUI_API FUIActionBindingHandle
 {
+	GENERATED_BODY()
+
 public:
 	bool IsValid() const;
 	void Unregister();
