PR #11238: CommonUI: Fix `UCommonGameViewportClient::HandleRerouteTouch` not getting called

diff --git a/Engine/Plugins/Runtime/CommonUI/Source/CommonUI/Private/CommonGameViewportClient.cpp b/Engine/Plugins/Runtime/CommonUI/Source/CommonUI/Private/CommonGameViewportClient.cpp
--- a/Engine/Plugins/Runtime/CommonUI/Source/CommonUI/Private/CommonGameViewportClient.cpp
+++ b/Engine/Plugins/Runtime/CommonUI/Source/CommonUI/Private/CommonGameViewportClient.cpp
@@ -87,7 +87,7 @@ bool UCommonGameViewportClient::InputTouch(FViewport* InViewport, int32 Controll
 #endif
 	{
 		FReply Result = FReply::Unhandled();
-		if (OnRerouteTouch().ExecuteIfBound(ControllerId, Handle, Type, TouchLocation, Result))
+		if (!OnRerouteTouch().ExecuteIfBound(ControllerId, Handle, Type, TouchLocation, Result))
 		{
 			HandleRerouteTouch(ControllerId, Handle, Type, TouchLocation, Result);
 		}
