PR #11150: [Common UI] Add the ability to disable automatic gamepad type detection

diff --git a/Engine/Plugins/Runtime/CommonUI/Source/CommonInput/Private/CommonInputPreprocessor.cpp b/Engine/Plugins/Runtime/CommonUI/Source/CommonInput/Private/CommonInputPreprocessor.cpp
--- a/Engine/Plugins/Runtime/CommonUI/Source/CommonInput/Private/CommonInputPreprocessor.cpp
+++ b/Engine/Plugins/Runtime/CommonUI/Source/CommonInput/Private/CommonInputPreprocessor.cpp
@@ -211,7 +211,9 @@ void FCommonInputPreprocessor::RefreshCurrentInputMethod(ECommonInputType InputM
 	InputSubsystem.SetCurrentInputType(InputMethod);
 
 	// Try to auto-detect the type of gamepad
-	if ((InputMethod == ECommonInputType::Gamepad) && UCommonInputPlatformSettings::Get()->CanChangeGamepadType())
+	if (InputMethod == ECommonInputType::Gamepad
+	    && ICommonInputModule::GetSettings().GetEnableAutomaticGamepadTypeDetection()
+	    && UCommonInputPlatformSettings::Get()->CanChangeGamepadType())
 	{
 		if (const FInputDeviceScope* DeviceScope = FInputDeviceScope::GetCurrent())
 		{
diff --git a/Engine/Plugins/Runtime/CommonUI/Source/CommonInput/Public/CommonInputSettings.h b/Engine/Plugins/Runtime/CommonUI/Source/CommonInput/Public/CommonInputSettings.h
--- a/Engine/Plugins/Runtime/CommonUI/Source/CommonInput/Public/CommonInputSettings.h
+++ b/Engine/Plugins/Runtime/CommonUI/Source/CommonInput/Public/CommonInputSettings.h
@@ -60,6 +60,8 @@ public:
 
 	bool GetEnableEnhancedInputSupport() const { return bEnableEnhancedInputSupport; }
 
+	bool GetEnableAutomaticGamepadTypeDetection() const { return bEnableAutomaticGamepadTypeDetection; }
+
 	TObjectPtr<UCommonInputActionDomainTable> GetActionDomainTable() const { return ActionDomainTablePtr; }
 
 public:
@@ -107,6 +109,13 @@ private:
 	UPROPERTY(config, EditAnywhere, Category = "Input", meta = (ConfigRestartRequired = true))
 	bool bEnableEnhancedInputSupport = false;
 
+	/**
+	* Controls automatic detection of the gamepad type.
+	* Disable this if you want to manually control the gamepad type using the UCommonInputSubsystem::SetGamepadInputType() function.
+	*/
+	UPROPERTY(config, EditAnywhere, Category = "Input")
+	bool bEnableAutomaticGamepadTypeDetection = true;
+
 	/** Create a derived asset from UCommonInputActionDomainTable to store ordered ActionDomain data for your game */
 	UPROPERTY(config, EditAnywhere, Category = "Action Domain")
 	TSoftObjectPtr<UCommonInputActionDomainTable> ActionDomainTable;
