PR #11133: Make UCommonActionWidget::UpdateActionWidget virtual
PR #11134: CommonUI: Split logic that determines whether the Icon will update in UCommonActionWidget::UpdateActionWidget into its own virtual function
PR #11147: Make OnActionProgress and OnActionComplete in UCommonActionWidget virtual

diff --git a/Engine/Plugins/Runtime/CommonUI/Source/CommonUI/Private/CommonActionWidget.cpp b/Engine/Plugins/Runtime/CommonUI/Source/CommonUI/Private/CommonActionWidget.cpp
--- a/Engine/Plugins/Runtime/CommonUI/Source/CommonUI/Private/CommonActionWidget.cpp
+++ b/Engine/Plugins/Runtime/CommonUI/Source/CommonUI/Private/CommonActionWidget.cpp
@@ -300,53 +300,38 @@ void UCommonActionWidget::UpdateActionWidget()
 		const UCommonInputSubsystem* CommonInputSubsystem = GetInputSubsystem();
 		if (IsDesignTime() || (GetGameInstance() && ensure(CommonInputSubsystem) && CommonInputSubsystem->ShouldShowInputKeys()))
 		{
-			const FCommonInputActionDataBase* InputActionData = GetInputActionData();
-			const bool bIsEnhancedInputAction = EnhancedInputAction && CommonUI::IsEnhancedInputSupportEnabled();
-
-#if WITH_EDITORONLY_DATA
-			const bool bIsDesignPreview = IsDesignTime() && DesignTimeKey.IsValid();
-#else
-			const bool bIsDesignPreview = false;
-#endif
-			if (InputActionData || bIsEnhancedInputAction || bIsDesignPreview)
+			if (ShouldUpdateActionWidgetIcon())
 			{
-				if (bAlwaysHideOverride)
+				Icon = GetIcon();
+
+				if (Icon.DrawAs == ESlateBrushDrawType::NoDrawType)
 				{
 					SetVisibility(ESlateVisibility::Collapsed);
 				}
-				else
+				else if (MyIcon.IsValid())
 				{
-					Icon = GetIcon();
+					MyIcon->SetImage(&Icon);
 
-					if (Icon.DrawAs == ESlateBrushDrawType::NoDrawType)
+					if (GetVisibility() != ESlateVisibility::Collapsed)
 					{
-						SetVisibility(ESlateVisibility::Collapsed);
+						// The object being passed into SetImage is the same each time so layout is never invalidated
+						// Manually invalidate it here as the dimensions may have changed
+						MyIcon->Invalidate(EInvalidateWidgetReason::Layout);
 					}
-					else if (MyIcon.IsValid())
+
+					if (IsHeldAction())
 					{
-						MyIcon->SetImage(&Icon);
-
-						if (GetVisibility() != ESlateVisibility::Collapsed)
-						{
-							// The object being passed into SetImage is the same each time so layout is never invalidated
-							// Manually invalidate it here as the dimensions may have changed
-							MyIcon->Invalidate(EInvalidateWidgetReason::Layout);
-						}
-
-						if (IsHeldAction())
-						{
-							MyProgressImage->SetVisibility(EVisibility::SelfHitTestInvisible);
-						}
-						else
-						{
-							MyProgressImage->SetVisibility(EVisibility::Collapsed);
-						}
-
-						MyKeyBox->Invalidate(EInvalidateWidget::LayoutAndVolatility);
-						SetVisibility(IsDesignTime() ? ESlateVisibility::Visible : ESlateVisibility::SelfHitTestInvisible);
-
-						return;
+						MyProgressImage->SetVisibility(EVisibility::SelfHitTestInvisible);
 					}
+					else
+					{
+						MyProgressImage->SetVisibility(EVisibility::Collapsed);
+					}
+
+					MyKeyBox->Invalidate(EInvalidateWidget::LayoutAndVolatility);
+					SetVisibility(IsDesignTime() ? ESlateVisibility::Visible : ESlateVisibility::SelfHitTestInvisible);
+
+					return;
 				}
 			}
 		}
@@ -355,6 +340,25 @@ void UCommonActionWidget::UpdateActionWidget()
 	}
 }
 
+bool UCommonActionWidget::ShouldUpdateActionWidgetIcon() const
+{
+	if (bAlwaysHideOverride)
+	{
+		return false;
+	}
+
+	const FCommonInputActionDataBase* InputActionData = GetInputActionData();
+	const bool bIsEnhancedInputAction = EnhancedInputAction && CommonUI::IsEnhancedInputSupportEnabled();
+
+#if WITH_EDITORONLY_DATA
+	const bool bIsDesignPreview = IsDesignTime() && DesignTimeKey.IsValid();
+#else
+	const bool bIsDesignPreview = false;
+#endif
+
+	return InputActionData || bIsEnhancedInputAction || bIsDesignPreview;
+}
+
 void UCommonActionWidget::ListenToInputMethodChanged(bool bListen)
 {
 	if (UCommonInputSubsystem* CommonInputSubsystem = GetInputSubsystem())
diff --git a/Engine/Plugins/Runtime/CommonUI/Source/CommonUI/Public/CommonActionWidget.h b/Engine/Plugins/Runtime/CommonUI/Source/CommonUI/Public/CommonActionWidget.h
--- a/Engine/Plugins/Runtime/CommonUI/Source/CommonUI/Public/CommonActionWidget.h
+++ b/Engine/Plugins/Runtime/CommonUI/Source/CommonUI/Public/CommonActionWidget.h
@@ -81,8 +81,8 @@ public:
 	virtual const FText GetPaletteCategory() override;
 #endif
 
-	void OnActionProgress(float HeldPercent);
-	void OnActionComplete();
+	virtual void OnActionProgress(float HeldPercent);
+	virtual void OnActionComplete();
 	void SetProgressMaterial(const FSlateBrush& InProgressMaterialBrush, const FName& InProgressMaterialParam);
 	void SetHidden(bool bAlwaysHidden);
 
@@ -115,7 +115,9 @@ protected:
 	UCommonInputSubsystem* GetInputSubsystem() const;
 	const FCommonInputActionDataBase* GetInputActionData() const;
 
-	void UpdateActionWidget();
+	virtual void UpdateActionWidget();
+
+	virtual bool ShouldUpdateActionWidgetIcon() const;
 
 	virtual void OnWidgetRebuilt() override;
 	
