PR #11152: [Common UI] Add some useful functions to UCommonUIActionRouterBase

diff --git a/Engine/Plugins/Runtime/CommonUI/Source/CommonUI/Private/Input/CommonUIActionRouterBase.cpp b/Engine/Plugins/Runtime/CommonUI/Source/CommonUI/Private/Input/CommonUIActionRouterBase.cpp
--- a/Engine/Plugins/Runtime/CommonUI/Source/CommonUI/Private/Input/CommonUIActionRouterBase.cpp
+++ b/Engine/Plugins/Runtime/CommonUI/Source/CommonUI/Private/Input/CommonUIActionRouterBase.cpp
@@ -1168,6 +1168,14 @@ void UCommonUIActionRouterBase::RegisterWidgetBindings(const FActivatableTreeNod
 	}
 }
 
+void UCommonUIActionRouterBase::RefreshActiveRootFocusRestorationTarget() const
+{
+	if (ActiveRootNode)
+	{
+		ActiveRootNode->RefreshCachedRestorationTarget();
+	}
+}
+
 void UCommonUIActionRouterBase::RefreshActiveRootFocus()
 {
 	if (ActiveRootNode)
@@ -1184,6 +1192,11 @@ void UCommonUIActionRouterBase::RefreshUIInputConfig()
 	}
 }
 
+TWeakPtr<FActivatableTreeRoot> UCommonUIActionRouterBase::GetActiveRoot() const
+{
+	return ActiveRootNode;
+}
+
 void UCommonUIActionRouterBase::SetActiveRoot(FActivatableTreeRootPtr NewActiveRoot)
 {
 	if (ActiveRootNode)
diff --git a/Engine/Plugins/Runtime/CommonUI/Source/CommonUI/Public/Input/CommonUIActionRouterBase.h b/Engine/Plugins/Runtime/CommonUI/Source/CommonUI/Public/Input/CommonUIActionRouterBase.h
--- a/Engine/Plugins/Runtime/CommonUI/Source/CommonUI/Public/Input/CommonUIActionRouterBase.h
+++ b/Engine/Plugins/Runtime/CommonUI/Source/CommonUI/Public/Input/CommonUIActionRouterBase.h
@@ -97,7 +97,7 @@ public:
 
 	UCommonInputSubsystem& GetInputSubsystem() const;
 
-	ERouteUIInputResult ProcessInput(FKey Key, EInputEvent InputEvent) const;
+	virtual ERouteUIInputResult ProcessInput(FKey Key, EInputEvent InputEvent) const;
 	bool CanProcessNormalGameInput() const;
 
 	bool IsPendingTreeChange() const;
@@ -125,6 +125,7 @@ public:
 
 	int32 GetLocalPlayerIndex() const;
 
+	void RefreshActiveRootFocusRestorationTarget() const;
 	void RefreshActiveRootFocus();
 	void RefreshUIInputConfig();
 
@@ -135,6 +136,7 @@ protected:
 	virtual void PostAnalogCursorCreate();
 	void RegisterAnalogCursorTick();
 
+	TWeakPtr<FActivatableTreeRoot> GetActiveRoot() const;
 	virtual void SetActiveRoot(FActivatableTreeRootPtr NewActiveRoot);
 	void SetForceResetActiveRoot(bool bInForceResetActiveRoot);
 
