#tests in Editor (PIE)

diff --git a/Engine/Plugins/Runtime/CommonUI/Source/CommonUI/Private/Input/CommonUIActionRouterBase.cpp b/Engine/Plugins/Runtime/CommonUI/Source/CommonUI/Private/Input/CommonUIActionRouterBase.cpp
--- a/Engine/Plugins/Runtime/CommonUI/Source/CommonUI/Private/Input/CommonUIActionRouterBase.cpp
+++ b/Engine/Plugins/Runtime/CommonUI/Source/CommonUI/Private/Input/CommonUIActionRouterBase.cpp
@@ -516,6 +516,22 @@ bool UCommonUIActionRouterBase::IsWidgetInActiveRoot(const UCommonActivatableWid
 	return false;
 }
 
+bool UCommonUIActionRouterBase::IsWidgetOfClassPresentInTree(TSubclassOf<UWidget> WidgetClass, bool IncludeChildren, bool IncludeInactive)
+{
+	bool IsPresent = false;
+
+	for (const FActivatableTreeRootRef& RootNode : RootNodes)
+	{
+		if(RootNode->IsWidgetOfClassPresentRecursive(WidgetClass, IncludeChildren, IncludeInactive))
+		{
+			IsPresent = true;
+			break;
+		}
+	}
+	
+	return IsPresent;
+}
+
 void UCommonUIActionRouterBase::NotifyUserWidgetConstructed(const UCommonUserWidget& Widget)
 {
 	check(Widget.GetCachedWidget().IsValid());
diff --git a/Engine/Plugins/Runtime/CommonUI/Source/CommonUI/Private/Input/UIActionRouterTypes.cpp b/Engine/Plugins/Runtime/CommonUI/Source/CommonUI/Private/Input/UIActionRouterTypes.cpp
--- a/Engine/Plugins/Runtime/CommonUI/Source/CommonUI/Private/Input/UIActionRouterTypes.cpp
+++ b/Engine/Plugins/Runtime/CommonUI/Source/CommonUI/Private/Input/UIActionRouterTypes.cpp
@@ -1161,6 +1161,34 @@ void FActivatableTreeNode::DebugDumpRecursive(FString& OutputStr, int32 Depth, b
 	}
 }
 
+bool FActivatableTreeNode::IsWidgetOfClassPresentRecursive(TSubclassOf<UWidget> WidgetClass, bool IncludeChildren, bool IncludeInactive) const
+{
+	if (!IncludeInactive && !IsWidgetActivated())
+	{
+		return false;
+	}
+
+	bool IsPresent = false;
+	if (RepresentedWidget.IsValid())
+	{
+		IsPresent |= RepresentedWidget->IsA(WidgetClass);
+	}
+
+	if (IncludeChildren && !IsPresent)
+	{
+		for (const FActivatableTreeNodeRef& ChildNode : Children)
+		{
+			if(ChildNode->IsWidgetOfClassPresentRecursive(WidgetClass, IncludeChildren, IncludeInactive))
+			{
+				IsPresent = true;
+				break;
+			}
+		}
+	}
+
+	return IsPresent;
+}
+
 bool FActivatableTreeNode::IsParentOfWidget(const TSharedPtr<SWidget>& SlateWidget) const
 {
 	if (SlateWidget && ensure(RepresentedWidget.IsValid()))
diff --git a/Engine/Plugins/Runtime/CommonUI/Source/CommonUI/Private/Input/UIActionRouterTypes.h b/Engine/Plugins/Runtime/CommonUI/Source/CommonUI/Private/Input/UIActionRouterTypes.h
--- a/Engine/Plugins/Runtime/CommonUI/Source/CommonUI/Private/Input/UIActionRouterTypes.h
+++ b/Engine/Plugins/Runtime/CommonUI/Source/CommonUI/Private/Input/UIActionRouterTypes.h
@@ -130,6 +130,8 @@ public:
 	void RemoveScrollRecipient(const UWidget& ScrollRecipient);
 	void AddInputPreprocessor(const TSharedRef<IInputProcessor>& InputPreprocessor, int32 DesiredIndex);
 
+	bool IsWidgetOfClassPresentRecursive(TSubclassOf<UWidget> WidgetClass, bool IncludeChildren, bool IncludeInactive) const;
+
 	FSimpleDelegate OnActivated;
 	FSimpleDelegate OnDeactivated;
 
diff --git a/Engine/Plugins/Runtime/CommonUI/Source/CommonUI/Public/Input/CommonUIActionRouterBase.h b/Engine/Plugins/Runtime/CommonUI/Source/CommonUI/Public/Input/CommonUIActionRouterBase.h
--- a/Engine/Plugins/Runtime/CommonUI/Source/CommonUI/Public/Input/CommonUIActionRouterBase.h
+++ b/Engine/Plugins/Runtime/CommonUI/Source/CommonUI/Public/Input/CommonUIActionRouterBase.h
@@ -107,6 +107,7 @@ public:
 	void FlushInput();
 
 	bool IsWidgetInActiveRoot(const UCommonActivatableWidget* Widget) const;
+	bool IsWidgetOfClassPresentInTree(TSubclassOf<UWidget> WidgetClass, bool bIncludeChildren, bool bIncludeInactive);
 
 	/** 
 	 * Sets Input Config 
