CommonTabListWidgetBase
- Updated CommonTabListWidgetBase to use a widget pool for the case of recreating tabs in this widget.
- It should alleviate the need to call create widget any time we're trying to build a new set of tabs. It also improves performance for that case.

diff --git a/Engine/Plugins/Runtime/CommonUI/Source/CommonUI/Private/CommonTabListWidgetBase.cpp b/Engine/Plugins/Runtime/CommonUI/Source/CommonUI/Private/CommonTabListWidgetBase.cpp
--- a/Engine/Plugins/Runtime/CommonUI/Source/CommonUI/Private/CommonTabListWidgetBase.cpp
+++ b/Engine/Plugins/Runtime/CommonUI/Source/CommonUI/Private/CommonTabListWidgetBase.cpp
@@ -17,6 +17,7 @@ UCommonTabListWidgetBase::UCommonTabListWidgetBase(const FObjectInitializer& Obj
 	, TabButtonGroup(nullptr)
 	, bIsListeningForInput(false)
 	, RegisteredTabsByID()
+	, TabButtonWidgetPool(*this)
 	, ActiveTabID(NAME_None)
 	, bIsRebuildingList(false)
 	, bPendingRebuild(false)
@@ -67,7 +68,7 @@ bool UCommonTabListWidgetBase::RegisterTab(FName TabNameID, TSubclassOf<UCommonB
 	}
 
 	// There is no PlayerController in Designer
-	UCommonButtonBase* const NewTabButton = IsDesignTime() ? CreateWidget<UCommonButtonBase>(GetWorld(), ButtonWidgetType) : CreateWidget<UCommonButtonBase>(GetOwningPlayer(), ButtonWidgetType);
+	UCommonButtonBase* const NewTabButton = TabButtonWidgetPool.GetOrCreateInstance<UCommonButtonBase>(ButtonWidgetType);
 	if (!ensureMsgf(NewTabButton, TEXT("Failed to create tab button. Aborting tab registration.")))
 	{
 		return false;
@@ -407,6 +408,13 @@ void UCommonTabListWidgetBase::NativeDestruct()
 	}
 }
 
+void UCommonTabListWidgetBase::ReleaseSlateResources(bool bReleaseChildren)
+{
+	Super::ReleaseSlateResources(bReleaseChildren);
+
+	TabButtonWidgetPool.ResetPool();
+}
+
 void UCommonTabListWidgetBase::HandlePreLinkedSwitcherChanged()
 {
 	HandlePreLinkedSwitcherChanged_BP();
@@ -532,6 +540,7 @@ void UCommonTabListWidgetBase::RemoveTab_Internal(const FName TabNameID, const F
 	{
 		TabButtonGroup->RemoveWidget(TabButton);
 		TabButton->RemoveFromParent();
+		TabButtonWidgetPool.Release(TabButton);
 	}
 
 	RegisteredTabsByID.Remove(TabNameID);
diff --git a/Engine/Plugins/Runtime/CommonUI/Source/CommonUI/Public/CommonTabListWidgetBase.h b/Engine/Plugins/Runtime/CommonUI/Source/CommonUI/Public/CommonTabListWidgetBase.h
--- a/Engine/Plugins/Runtime/CommonUI/Source/CommonUI/Public/CommonTabListWidgetBase.h
+++ b/Engine/Plugins/Runtime/CommonUI/Source/CommonUI/Public/CommonTabListWidgetBase.h
@@ -2,9 +2,10 @@
 
 #pragma once
 
+#include "Blueprint/UserWidgetPool.h"
 #include "CommonUserWidget.h"
-
 #include "Engine/DataTable.h"
+
 #include "CommonTabListWidgetBase.generated.h"
 
 class UCommonAnimatedSwitcher;
@@ -150,6 +151,7 @@ protected:
 	virtual void NativeOnInitialized() override;
 	virtual void NativeConstruct() override;
 	virtual void NativeDestruct() override;
+	virtual void ReleaseSlateResources(bool bReleaseChildren) override;
 	// End UUserWidget
 
 	virtual void UpdateBindings();
@@ -235,6 +237,9 @@ private:
 	UPROPERTY(Transient)
 	TMap<FName, FCommonRegisteredTabInfo> RegisteredTabsByID;
 
+	UPROPERTY(Transient)
+	FUserWidgetPool TabButtonWidgetPool;
+
 	/** The registration ID of the currently active tab */
 	FName ActiveTabID;
 
