PR #11449: Fix typo in log format string when a widget gets deactivated

diff --git a/Engine/Plugins/Runtime/CommonUI/Source/CommonUI/Private/CommonActivatableWidget.cpp b/Engine/Plugins/Runtime/CommonUI/Source/CommonUI/Private/CommonActivatableWidget.cpp
--- a/Engine/Plugins/Runtime/CommonUI/Source/CommonUI/Private/CommonActivatableWidget.cpp
+++ b/Engine/Plugins/Runtime/CommonUI/Source/CommonUI/Private/CommonActivatableWidget.cpp
@@ -295,7 +295,7 @@ void UCommonActivatableWidget::NativeOnDeactivated()
 		if (bSetVisibilityOnDeactivated)
 		{
 			SetVisibility(DeactivatedVisibility);
-			UE_LOG(LogCommonUI, Verbose, TEXT("[%s] set visibility to [%d] on deactivation"), *GetName(), *StaticEnum<ESlateVisibility>()->GetDisplayValueAsText(DeactivatedVisibility).ToString());
+			UE_LOG(LogCommonUI, Verbose, TEXT("[%s] set visibility to [%s] on deactivation"), *GetName(), *StaticEnum<ESlateVisibility>()->GetDisplayValueAsText(DeactivatedVisibility).ToString());
 		}
 
 		if (CommonUI::IsEnhancedInputSupportEnabled() && InputMapping)
