Fix log spam in CommonUI that was happening on every garbage collect

diff --git a/Engine/Plugins/Runtime/CommonUI/Source/CommonUI/Private/Input/UIActionRouterTypes.cpp b/Engine/Plugins/Runtime/CommonUI/Source/CommonUI/Private/Input/UIActionRouterTypes.cpp
--- a/Engine/Plugins/Runtime/CommonUI/Source/CommonUI/Private/Input/UIActionRouterTypes.cpp
+++ b/Engine/Plugins/Runtime/CommonUI/Source/CommonUI/Private/Input/UIActionRouterTypes.cpp
@@ -280,7 +280,10 @@ void FUIActionBinding::CleanRegistrations()
 		}
 	}
 
-	UE_LOG(LogUIActionRouter, Log, TEXT("Cleaned out [%d] inactive UI action bindings"), NumRemoved);
+	if (NumRemoved > 0)
+	{
+		UE_LOG(LogUIActionRouter, Log, TEXT("Cleaned out [%d] inactive UI action bindings"), NumRemoved);
+	}
 }
 
 FCommonInputActionDataBase* FUIActionBinding::GetLegacyInputActionData() const
