UMG: Make sure that CommonUI button object remains consistent after changing sound override properties

diff --git a/Engine/Plugins/Runtime/CommonUI/Source/CommonUI/Private/CommonButtonBase.cpp b/Engine/Plugins/Runtime/CommonUI/Source/CommonUI/Private/CommonButtonBase.cpp
--- a/Engine/Plugins/Runtime/CommonUI/Source/CommonUI/Private/CommonButtonBase.cpp
+++ b/Engine/Plugins/Runtime/CommonUI/Source/CommonUI/Private/CommonButtonBase.cpp
@@ -1762,6 +1762,66 @@ void UCommonButtonBase::SetLockedHoveredSoundOverride(USoundBase* Sound)
 	}
 }
 
+void UCommonButtonBase::SetPressedSlateSoundOverride(const FSlateSound& InPressedSlateSoundOverride)
+{
+	const bool bBuildStyles = InPressedSlateSoundOverride.GetResourceObject() != PressedSlateSoundOverride.GetResourceObject();
+	PressedSlateSoundOverride = InPressedSlateSoundOverride;
+	if (bBuildStyles)
+	{
+		BuildStyles();
+	}
+}
+
+void UCommonButtonBase::SetHoveredSlateSoundOverride(const FSlateSound& InHoveredSlateSoundOverride)
+{
+	const bool bBuildStyles = InHoveredSlateSoundOverride.GetResourceObject() != HoveredSlateSoundOverride.GetResourceObject();
+	HoveredSlateSoundOverride = InHoveredSlateSoundOverride;
+	if (bBuildStyles)
+	{
+		BuildStyles();
+	}
+}
+
+void UCommonButtonBase::SetSelectedPressedSlateSoundOverride(const FSlateSound& InSelectedPressedSlateSoundOverride)
+{
+	const bool bBuildStyles = InSelectedPressedSlateSoundOverride.GetResourceObject() != SelectedPressedSlateSoundOverride.GetResourceObject();
+	SelectedPressedSlateSoundOverride = InSelectedPressedSlateSoundOverride;
+	if (bBuildStyles)
+	{
+		BuildStyles();
+	}
+}
+
+void UCommonButtonBase::SetSelectedHoveredSlateSoundOverride(const FSlateSound& InSelectedHoveredSlateSoundOverride)
+{
+	const bool bBuildStyles = InSelectedHoveredSlateSoundOverride.GetResourceObject() != SelectedHoveredSlateSoundOverride.GetResourceObject();
+	SelectedHoveredSlateSoundOverride = InSelectedHoveredSlateSoundOverride;
+	if (bBuildStyles)
+	{
+		BuildStyles();
+	}
+}
+
+void UCommonButtonBase::SetLockedPressedSlateSoundOverride(const FSlateSound& InLockedPressedSlateSoundOverride)
+{
+	const bool bBuildStyles = InLockedPressedSlateSoundOverride.GetResourceObject() != LockedPressedSlateSoundOverride.GetResourceObject();
+	LockedPressedSlateSoundOverride = InLockedPressedSlateSoundOverride;
+	if (bBuildStyles)
+	{
+		BuildStyles();
+	}
+}
+
+void UCommonButtonBase::SetLockedHoveredSlateSoundOverride(const FSlateSound& InLockedHoveredSlateSoundOverride)
+{
+	const bool bBuildStyles = InLockedHoveredSlateSoundOverride.GetResourceObject() != LockedHoveredSlateSoundOverride.GetResourceObject();
+	LockedHoveredSlateSoundOverride = InLockedHoveredSlateSoundOverride;
+	if (bBuildStyles)
+	{
+		BuildStyles();
+	}
+}
+
 void UCommonButtonBase::UpdateInputActionWidgetVisibility()
 {
 	if (InputActionWidget)
diff --git a/Engine/Plugins/Runtime/CommonUI/Source/CommonUI/Public/CommonButtonBase.h b/Engine/Plugins/Runtime/CommonUI/Source/CommonUI/Public/CommonButtonBase.h
--- a/Engine/Plugins/Runtime/CommonUI/Source/CommonUI/Public/CommonButtonBase.h
+++ b/Engine/Plugins/Runtime/CommonUI/Source/CommonUI/Public/CommonButtonBase.h
@@ -500,6 +500,14 @@ protected:
 	const FText GetPaletteCategory() override;
 #endif // WITH_EDITOR
 
+	// In some scenarios, USoundBase* setters are bypassed, so standard setters must also be provided for sound override properties
+	void SetPressedSlateSoundOverride(const FSlateSound& InPressedSlateSoundOverride);
+	void SetHoveredSlateSoundOverride(const FSlateSound& InHoveredSlateSoundOverride);
+	void SetSelectedPressedSlateSoundOverride(const FSlateSound& InSelectedPressedSlateSoundOverride);
+	void SetSelectedHoveredSlateSoundOverride(const FSlateSound& InSelectedHoveredSlateSoundOverride);
+	void SetLockedPressedSlateSoundOverride(const FSlateSound& InLockedPressedSlateSoundOverride);
+	void SetLockedHoveredSlateSoundOverride(const FSlateSound& InLockedHoveredSlateSoundOverride);
+
 	/** Helper function to bind to input method change events */
 	virtual void BindInputMethodChangedDelegate();
 
@@ -700,30 +708,30 @@ protected:
 	 * Optional override for the sound to play when this button is pressed.
 	 * Also used for the Selected and Locked Pressed state if their respective overrides are empty.
 	 */
-	UPROPERTY(EditAnywhere, BlueprintReadOnly, Category = Sound, meta = (DisplayName = "Pressed Sound Override"))
+	UPROPERTY(EditAnywhere, BlueprintReadOnly, Setter, Category = Sound, meta = (DisplayName = "Pressed Sound Override"))
 	FSlateSound PressedSlateSoundOverride;
 
 	/**
 	 * Optional override for the sound to play when this button is hovered.
 	 * Also used for the Selected and Locked Hovered state if their respective overrides are empty.
 	 */
-	UPROPERTY(EditAnywhere, BlueprintReadOnly, Category = Sound, meta = (DisplayName = "Hovered Sound Override"))
+	UPROPERTY(EditAnywhere, BlueprintReadOnly, Setter, Category = Sound, meta = (DisplayName = "Hovered Sound Override"))
 	FSlateSound HoveredSlateSoundOverride;
 
 	/** Optional override for the sound to play when this button is pressed while Selected */
-	UPROPERTY(EditAnywhere, BlueprintReadOnly, Category = Sound, meta = (DisplayName = "Selected Pressed Sound Override"))
+	UPROPERTY(EditAnywhere, BlueprintReadOnly, Setter, Category = Sound, meta = (DisplayName = "Selected Pressed Sound Override"))
 	FSlateSound SelectedPressedSlateSoundOverride;
 
 	/** Optional override for the sound to play when this button is hovered while Selected */
-	UPROPERTY(EditAnywhere, BlueprintReadOnly, Category = Sound, meta = (DisplayName = "Selected Hovered Sound Override"))
+	UPROPERTY(EditAnywhere, BlueprintReadOnly, Setter, Category = Sound, meta = (DisplayName = "Selected Hovered Sound Override"))
 	FSlateSound SelectedHoveredSlateSoundOverride;
 
 	/** Optional override for the sound to play when this button is pressed while Locked */
-	UPROPERTY(EditAnywhere, BlueprintReadOnly, Category = Sound, meta = (DisplayName = "Locked Pressed Sound Override"))
+	UPROPERTY(EditAnywhere, BlueprintReadOnly, Setter, Category = Sound, meta = (DisplayName = "Locked Pressed Sound Override"))
 	FSlateSound LockedPressedSlateSoundOverride;
 
 	/** Optional override for the sound to play when this button is hovered while Locked */
-	UPROPERTY(EditAnywhere, BlueprintReadOnly, Category = Sound, meta = (DisplayName = "Locked Hovered Sound Override"))
+	UPROPERTY(EditAnywhere, BlueprintReadOnly, Setter, Category = Sound, meta = (DisplayName = "Locked Hovered Sound Override"))
 	FSlateSound LockedHoveredSlateSoundOverride;
 
 	/** The type of mouse action required by the user to trigger the button's 'Click' */
