Fix ensure firing on selection change during destruct.

#tests Lyra PIE tabbing

diff --git a/Engine/Plugins/Runtime/CommonUI/Source/CommonUI/Private/Groups/CommonButtonGroupBase.cpp b/Engine/Plugins/Runtime/CommonUI/Source/CommonUI/Private/Groups/CommonButtonGroupBase.cpp
--- a/Engine/Plugins/Runtime/CommonUI/Source/CommonUI/Private/Groups/CommonButtonGroupBase.cpp
+++ b/Engine/Plugins/Runtime/CommonUI/Source/CommonUI/Private/Groups/CommonButtonGroupBase.cpp
@@ -272,11 +272,16 @@ void UCommonButtonGroupBase::OnRemoveAll()
 	// end up with a case of the selection change event triggering multiple times when removing all.
 	SelectedButtonIndex = INDEX_NONE;
 
-	while ( Buttons.Num() > 0 )
 	{
-		if (UCommonButtonBase* Button = Buttons.Pop().Get())
+		// Even if we require selection, allow for it be ignored when removing all buttons (Ex: Destructor).
+		TGuardValue<bool> DisableSelectionRequiredGuard(bSelectionRequired, false);
+
+		while (Buttons.Num() > 0)
 		{
-			RemoveWidget(Button);
+			if (UCommonButtonBase* Button = Buttons.Pop().Get())
+			{
+				RemoveWidget(Button);
+			}
 		}
 	}
 	
