[CommonUI] Adding transition fallback capability to CommonAnimatedSwitcher to handle cases when relevant slots are removed during a transition.

#changelist validated
#virtualized

diff --git a/Engine/Plugins/Runtime/CommonUI/Source/CommonUI/Private/CommonAnimatedSwitcher.cpp b/Engine/Plugins/Runtime/CommonUI/Source/CommonUI/Private/CommonAnimatedSwitcher.cpp
--- a/Engine/Plugins/Runtime/CommonUI/Source/CommonUI/Private/CommonAnimatedSwitcher.cpp
+++ b/Engine/Plugins/Runtime/CommonUI/Source/CommonUI/Private/CommonAnimatedSwitcher.cpp
@@ -117,6 +117,7 @@ PRAGMA_DISABLE_DEPRECATION_WARNINGS
 		.TransitionCurveType(TransitionCurveType)
 		.TransitionDuration(TransitionDuration)
 		.TransitionType(TransitionType)
+		.TransitionFallbackStrategy(TransitionFallbackStrategy)
 		.OnActiveIndexChanged_UObject(this, &UCommonAnimatedSwitcher::HandleSlateActiveIndexChanged)
 		.OnIsTransitioningChanged_UObject(this, &UCommonAnimatedSwitcher::HandleSlateIsTransitioningChanged);
 PRAGMA_ENABLE_DEPRECATION_WARNINGS
@@ -205,3 +206,26 @@ bool UCommonAnimatedSwitcher::IsTransitionPlaying() const
 	}
 }
 
+UWidget* UCommonAnimatedSwitcher::GetPendingActiveWidget() const
+{
+	if (MyAnimatedSwitcher.IsValid())
+	{
+		const int32 PendingIndex = GetPendingActiveWidgetIndex();
+		if (PendingIndex >= 0)
+		{
+			return GetWidgetAtIndex(PendingIndex);
+		}
+	}
+	
+	return nullptr;
+}
+
+int32 UCommonAnimatedSwitcher::GetPendingActiveWidgetIndex() const
+{
+	if (MyAnimatedSwitcher.IsValid())
+	{
+		return MyAnimatedSwitcher->GetPendingActiveWidgetIndex();
+	}
+	return INDEX_NONE;
+}
+
diff --git a/Engine/Plugins/Runtime/CommonUI/Source/CommonUI/Private/Slate/SCommonAnimatedSwitcher.cpp b/Engine/Plugins/Runtime/CommonUI/Source/CommonUI/Private/Slate/SCommonAnimatedSwitcher.cpp
--- a/Engine/Plugins/Runtime/CommonUI/Source/CommonUI/Private/Slate/SCommonAnimatedSwitcher.cpp
+++ b/Engine/Plugins/Runtime/CommonUI/Source/CommonUI/Private/Slate/SCommonAnimatedSwitcher.cpp
@@ -11,8 +11,10 @@ void SCommonAnimatedSwitcher::Construct(const FArguments& InArgs)
 	SetCanTick(false);
 	bTransitioningOut = false;
 	PendingActiveWidget = nullptr;
+	PendingActiveWidgetIndex = INDEX_NONE;
 
 	TransitionType = InArgs._TransitionType;
+	TransitionFallbackStrategy = InArgs._TransitionFallbackStrategy;
 
 	SetTransition(InArgs._TransitionDuration, InArgs._TransitionCurveType);
 
@@ -45,8 +47,13 @@ void SCommonAnimatedSwitcher::TransitionToIndex(int32 NewWidgetIndex, bool bInst
 {
 	// Cache the widget we want to reach
 	PendingActiveWidget = GetWidget(NewWidgetIndex);
-	if (!PendingActiveWidget.IsValid())
+	if (PendingActiveWidget.IsValid())
 	{
+		PendingActiveWidgetIndex = NewWidgetIndex;
+	}
+	else
+	{
+		PendingActiveWidgetIndex = INDEX_NONE;
 		UE_LOG(LogSlate, Verbose, TEXT("Called SCommonAnimatedSwitcher::TransitionToIndex('%d') to an invalid Index"), NewWidgetIndex);
 		return;
 	}
@@ -134,6 +141,55 @@ void SCommonAnimatedSwitcher::SetTransition(float Duration, ETransitionCurve Cur
 	TransitionSequence = FCurveSequence(0.f, Duration * 0.5f, TransitionCurveToCurveEaseFunction(Curve));
 }
 
+void SCommonAnimatedSwitcher::OnSlotAdded(int32 AddedIndex)
+{
+	SWidgetSwitcher::OnSlotAdded(AddedIndex);
+	
+	if (PendingActiveWidgetIndex >= 0 && PendingActiveWidgetIndex >= AddedIndex)
+	{
+		// A slot has been added before our PendingActiveWidgetIndex
+		++PendingActiveWidgetIndex;
+	}
+}
+
+void SCommonAnimatedSwitcher::OnSlotRemoved(int32 RemovedIndex, TSharedRef<SWidget> RemovedWidget, bool bWasActiveSlot)
+{
+	SWidgetSwitcher::OnSlotRemoved(RemovedIndex, RemovedWidget, bWasActiveSlot);
+
+	if (PendingActiveWidgetIndex >= 0)
+	{
+		if (PendingActiveWidgetIndex == RemovedIndex)
+		{
+			// Our PendingActiveWidget has been removed, attempt a fallback if enabled
+			TryTransitionFallbackOfPendingWidget();
+		}
+		else if (PendingActiveWidgetIndex < RemovedIndex)
+		{
+			// A slot has been removed before our PendingActiveWidget
+			--PendingActiveWidgetIndex;
+		}
+	}
+	
+	if (bWasActiveSlot && TransitionSequence.IsPlaying())
+	{
+		// The widget being animated has been removed. Instantly finish the current transition animation and allow UpdateTransition() to finish the process.
+		if (TransitionSequence.IsForward())
+		{
+			TransitionSequence.JumpToEnd();
+		}
+		else
+		{
+			TransitionSequence.JumpToStart();
+		}
+
+		if (!bTransitioningOut)
+		{
+			// Since the removed widget was active and the target of the current transition, set the fallback widget active.
+			TryTransitionFallbackOfActiveWidget(RemovedIndex);
+		}
+	}
+}
+
 EActiveTimerReturnType SCommonAnimatedSwitcher::UpdateTransition(double InCurrentTime, float InDeltaTime)
 {
 	if (TransitionType == ECommonSwitcherTransition::Zoom)
@@ -163,7 +219,19 @@ EActiveTimerReturnType SCommonAnimatedSwitcher::UpdateTransition(double InCurren
 
 	if (!TransitionSequence.IsPlaying())
 	{
+		// The transition is complete in one direction
 		TSharedPtr<SWidget> PinnedPendingActiveWidget = PendingActiveWidget.Pin();
+		if (bTransitioningOut && (!PinnedPendingActiveWidget.IsValid() || GetWidgetIndex(PinnedPendingActiveWidget.ToSharedRef()) == INDEX_NONE))
+		{
+			// The PendingActiveWidget is no longer contained in this switcher
+			UE_LOG(LogSlate, Verbose, TEXT("SCommonAnimatedSwitcher Pending Widget was removed or became invalid while a transition was happening"));
+			if (TryTransitionFallbackOfPendingWidget())
+			{
+				// We have fallen back to an alternative PendingActiveWidget
+				PinnedPendingActiveWidget = PendingActiveWidget.Pin();
+			}
+		}
+		
 		if (PinnedPendingActiveWidget.IsValid() && GetActiveWidget().Get() != PinnedPendingActiveWidget.Get())
 		{
 			const bool bWasTransitioningOut = bTransitioningOut;
@@ -174,6 +242,7 @@ EActiveTimerReturnType SCommonAnimatedSwitcher::UpdateTransition(double InCurren
 				// Finished transitioning out - update the active widget and play again from the start
 				SetActiveWidget(PinnedPendingActiveWidget.ToSharedRef());
 				PendingActiveWidget = nullptr;
+				PendingActiveWidgetIndex = INDEX_NONE;
 
 				//Note that any listener could decide to trigger ANOTHER transition, 
 				//changing the values of PendingActiveWidget and bTransitioningOut to something new.
@@ -196,11 +265,8 @@ EActiveTimerReturnType SCommonAnimatedSwitcher::UpdateTransition(double InCurren
 		}
 		else
 		{
-			if (!PinnedPendingActiveWidget.IsValid())
-			{
-				UE_LOG(LogSlate, Verbose, TEXT("SCommonAnimatedSwitcher Pending Widget became invalid while a transition was happening"));
-			}
 			PendingActiveWidget = nullptr;
+			PendingActiveWidgetIndex = INDEX_NONE;
 		}
 	}
 
@@ -235,6 +301,73 @@ float SCommonAnimatedSwitcher::GetTransitionProgress() const
 	return Progress;
 }
 
+int32 SCommonAnimatedSwitcher::GetTransitionFallbackForIndex(int32 RemovedWidgetIndex) const
+{
+	int32 DesiredIndex;
+	switch (TransitionFallbackStrategy)
+	{
+	case ECommonSwitcherTransitionFallbackStrategy::First:
+		// Falling back to the index of the first item
+		DesiredIndex = 0;
+		break;
+	case ECommonSwitcherTransitionFallbackStrategy::Last:
+		// Falling back to the index of the last item
+		DesiredIndex = GetNumWidgets() - 1;
+		break;
+	case ECommonSwitcherTransitionFallbackStrategy::Next:
+		// Falling back to the index of the next item, which is now the same as RemovedWidgetIndex because it shifted after the removal.
+		DesiredIndex = RemovedWidgetIndex;
+		break;
+	case ECommonSwitcherTransitionFallbackStrategy::Previous:
+	default:
+		// Falling back to the index of the previous item
+		DesiredIndex = RemovedWidgetIndex - 1;
+		break;
+	}
+
+	// Constraining the index to the valid range
+	return FMath::Clamp(DesiredIndex, 0, GetNumWidgets() - 1);
+}
+
+bool SCommonAnimatedSwitcher::TryTransitionFallbackOfPendingWidget()
+{
+	if (IsTransitionFallbackEnabled() && GetNumWidgets() > 0)
+	{
+		// Falling back to another transition target
+		const int32 NewPendingWidgetIndex = GetTransitionFallbackForIndex(PendingActiveWidgetIndex);
+		PendingActiveWidget = GetWidget(NewPendingWidgetIndex);
+		if (PendingActiveWidget.IsValid())
+		{
+			PendingActiveWidgetIndex = NewPendingWidgetIndex;
+			return true;
+		}
+		else
+		{
+			PendingActiveWidgetIndex = INDEX_NONE;
+			UE_LOG(LogSlate, Verbose, TEXT("SCommonAnimatedSwitcher failed to fall back to another Pending Widget for the transition"));
+		}
+	}
+	
+	return false;
+}
+
+bool SCommonAnimatedSwitcher::TryTransitionFallbackOfActiveWidget(int32 RemovedWidgetIndex)
+{
+	if (IsTransitionFallbackEnabled() && GetNumWidgets() > 0)
+	{
+		// Falling back to another active widget
+		const int32 NewActiveWidgetIndex = GetTransitionFallbackForIndex(RemovedWidgetIndex);
+		SetActiveWidgetIndex(NewActiveWidgetIndex);
+		
+		//Note that any listener could decide to trigger ANOTHER transition, 
+		//changing the values of PendingActiveWidget and bTransitioningOut to something new.
+		OnActiveIndexChanged.ExecuteIfBound(GetActiveWidgetIndex());
+		return true;
+	}
+	
+	return false;
+}
+
 bool SCommonAnimatedSwitcher::IsTransitionPlaying() const
 {
 	return TransitionSequence.IsPlaying();
diff --git a/Engine/Plugins/Runtime/CommonUI/Source/CommonUI/Private/Widgets/CommonActivatableWidgetContainer.cpp b/Engine/Plugins/Runtime/CommonUI/Source/CommonUI/Private/Widgets/CommonActivatableWidgetContainer.cpp
--- a/Engine/Plugins/Runtime/CommonUI/Source/CommonUI/Private/Widgets/CommonActivatableWidgetContainer.cpp
+++ b/Engine/Plugins/Runtime/CommonUI/Source/CommonUI/Private/Widgets/CommonActivatableWidgetContainer.cpp
@@ -98,6 +98,7 @@ TSharedRef<SWidget> UCommonActivatableWidgetContainerBase::RebuildWidget()
 			.TransitionCurveType(TransitionCurveType)
 			.TransitionDuration(TransitionDuration)
 			.TransitionType(TransitionType)
+			.TransitionFallbackStrategy(TransitionFallbackStrategy)
 			.OnActiveIndexChanged_UObject(this, &UCommonActivatableWidgetContainerBase::HandleActiveIndexChanged)
 			.OnIsTransitioningChanged_UObject(this, &UCommonActivatableWidgetContainerBase::HandleSwitcherIsTransitioningChanged)
 		]
diff --git a/Engine/Plugins/Runtime/CommonUI/Source/CommonUI/Public/CommonAnimatedSwitcher.h b/Engine/Plugins/Runtime/CommonUI/Source/CommonUI/Public/CommonAnimatedSwitcher.h
--- a/Engine/Plugins/Runtime/CommonUI/Source/CommonUI/Public/CommonAnimatedSwitcher.h
+++ b/Engine/Plugins/Runtime/CommonUI/Source/CommonUI/Public/CommonAnimatedSwitcher.h
@@ -52,6 +52,9 @@ public:
 	UFUNCTION(BlueprintCallable, Category = "Common Widget Switcher")
 	bool IsTransitionPlaying() const;
 
+	UWidget* GetPendingActiveWidget() const;
+	int32 GetPendingActiveWidgetIndex() const;
+
 protected:
 	virtual void HandleSlateActiveIndexChanged(int32 ActiveIndex);
 
@@ -87,6 +90,10 @@ protected:
 	UPROPERTY(EditAnywhere, BlueprintReadOnly, Category = "Transition")
 	float TransitionDuration;
 
+	/** Controls how we will choose another widget if a transitioning widget is removed during the transition. */
+	UPROPERTY(EditAnywhere, BlueprintReadOnly, Category = "Transition")
+	ECommonSwitcherTransitionFallbackStrategy TransitionFallbackStrategy = ECommonSwitcherTransitionFallbackStrategy::None;
+
 	TSharedPtr<SOverlay> MyOverlay;
 	TSharedPtr<SSpacer> MyInputGuard;
 	TSharedPtr<SCommonAnimatedSwitcher> MyAnimatedSwitcher;
diff --git a/Engine/Plugins/Runtime/CommonUI/Source/CommonUI/Public/Slate/SCommonAnimatedSwitcher.h b/Engine/Plugins/Runtime/CommonUI/Source/CommonUI/Public/Slate/SCommonAnimatedSwitcher.h
--- a/Engine/Plugins/Runtime/CommonUI/Source/CommonUI/Public/Slate/SCommonAnimatedSwitcher.h
+++ b/Engine/Plugins/Runtime/CommonUI/Source/CommonUI/Public/Slate/SCommonAnimatedSwitcher.h
@@ -39,6 +39,22 @@ enum class ETransitionCurve : uint8
 	CubicInOut,
 };
 
+/** Determines the switcher's behavior if the target of a transition is removed before it becomes the active widget. */
+UENUM(BlueprintType)
+enum class ECommonSwitcherTransitionFallbackStrategy : uint8
+{
+	/** Transition fallbacks are disabled and no special handling will occur if a transitioning widget is removed. */
+	None,
+	/** Fall back to the nearest valid slot at a lower index than the original target, or the first slot if there are none lower */
+	Previous,
+	/** Fall back to the nearest valid slot at a higher index than the original target, or the last slot if there are none higher */
+	Next,
+	/** Fall back to the first item in the switcher */
+	First,
+	/** Fall back to the last item in the switcher */
+	Last
+};
+
 static FORCEINLINE ECurveEaseFunction TransitionCurveToCurveEaseFunction(ETransitionCurve CurveType)
 {
 	switch (CurveType)
@@ -65,6 +81,7 @@ public:
 		, _TransitionType(ECommonSwitcherTransition::FadeOnly)
 		, _TransitionCurveType(ETransitionCurve::CubicInOut)
 		, _TransitionDuration(0.4f)
+		, _TransitionFallbackStrategy(ECommonSwitcherTransitionFallbackStrategy::Previous)
 		{
 			_Visibility = EVisibility::SelfHitTestInvisible;
 		}
@@ -74,6 +91,7 @@ public:
 		SLATE_ARGUMENT(ECommonSwitcherTransition, TransitionType)
 		SLATE_ARGUMENT(ETransitionCurve, TransitionCurveType)
 		SLATE_ARGUMENT(float, TransitionDuration)
+		SLATE_ARGUMENT(ECommonSwitcherTransitionFallbackStrategy, TransitionFallbackStrategy)
 		SLATE_EVENT(FOnActiveIndexChanged, OnActiveIndexChanged)
 		SLATE_EVENT(FOnIsTransitioningChanged, OnIsTransitioningChanged)
 
@@ -90,10 +108,24 @@ public:
 	void SetTransition(float Duration, ETransitionCurve Curve);
 
 	bool IsTransitionPlaying() const;
+	
+	TWeakPtr<SWidget> GetPendingActiveWidget() const { return PendingActiveWidget; }
+	int32 GetPendingActiveWidgetIndex() const { return PendingActiveWidgetIndex; }
 
+	void SetTransitionFallbackStrategy(const ECommonSwitcherTransitionFallbackStrategy InStrategy) { TransitionFallbackStrategy = InStrategy; }
+	ECommonSwitcherTransitionFallbackStrategy GetTransitionFallbackStrategy() const { return TransitionFallbackStrategy; }
+	bool IsTransitionFallbackEnabled() const { return TransitionFallbackStrategy != ECommonSwitcherTransitionFallbackStrategy::None; }
+
+protected:
+	virtual void OnSlotAdded(int32 AddedIndex) override;
+	virtual void OnSlotRemoved(int32 RemovedIndex, TSharedRef<SWidget> RemovedWidget, bool bWasActiveSlot) override;
+	
 private:
 	EActiveTimerReturnType UpdateTransition(double InCurrentTime, float InDeltaTime);
 	float GetTransitionProgress() const;
+	int32 GetTransitionFallbackForIndex(int32 RemovedWidgetIndex) const;
+	bool TryTransitionFallbackOfPendingWidget();
+	bool TryTransitionFallbackOfActiveWidget(int32 RemovedWidgetIndex);
 
 private:
 	/** Anim sequence for the transition; plays twice per transition */
@@ -102,6 +134,12 @@ private:
 	/** The pending active widget, set when the initial transition out completes. If set to a null value then we don't have any pending widget */
 	TWeakPtr<SWidget> PendingActiveWidget;
 
+	/** Tracks the index of PendingActiveWidget, if there is one. */
+	int32 PendingActiveWidgetIndex = INDEX_NONE;
+
+	/** Controls how we will choose another widget if a transitioning widget is removed during the transition. */
+	ECommonSwitcherTransitionFallbackStrategy TransitionFallbackStrategy = ECommonSwitcherTransitionFallbackStrategy::None;
+
 	/** If true, we are transitioning content out and need to play the sequence again to transition it in */
 	bool bTransitioningOut = false;
 
diff --git a/Engine/Plugins/Runtime/CommonUI/Source/CommonUI/Public/Widgets/CommonActivatableWidgetContainer.h b/Engine/Plugins/Runtime/CommonUI/Source/CommonUI/Public/Widgets/CommonActivatableWidgetContainer.h
--- a/Engine/Plugins/Runtime/CommonUI/Source/CommonUI/Public/Widgets/CommonActivatableWidgetContainer.h
+++ b/Engine/Plugins/Runtime/CommonUI/Source/CommonUI/Public/Widgets/CommonActivatableWidgetContainer.h
@@ -3,6 +3,7 @@
 #pragma once
 
 #include "Blueprint/UserWidgetPool.h"
+#include "Slate/SCommonAnimatedSwitcher.h"
 #include "CommonActivatableWidgetContainer.generated.h"
 
 class SCommonAnimatedSwitcher;
@@ -108,6 +109,13 @@ protected:
 	UPROPERTY(EditAnywhere, Category = Transition)
 	float TransitionDuration = 0.4f;
 
+	/**
+	 * Controls how we will choose another widget if a transitioning widget is removed during the transition.
+	 * Note for Queues and Stacks, ECommonSwitcherTransitionFallbackStrategy::Previous is a good option.
+	 */
+	UPROPERTY(EditAnywhere, BlueprintReadOnly, Category = Transition)
+	ECommonSwitcherTransitionFallbackStrategy TransitionFallbackStrategy = ECommonSwitcherTransitionFallbackStrategy::None;
+
 	UPROPERTY(Transient)
 	TArray<TObjectPtr<UCommonActivatableWidget>> WidgetList;
 
diff --git a/Engine/Source/Runtime/Slate/Private/Widgets/Layout/SWidgetSwitcher.cpp b/Engine/Source/Runtime/Slate/Private/Widgets/Layout/SWidgetSwitcher.cpp
--- a/Engine/Source/Runtime/Slate/Private/Widgets/Layout/SWidgetSwitcher.cpp
+++ b/Engine/Source/Runtime/Slate/Private/Widgets/Layout/SWidgetSwitcher.cpp
@@ -24,14 +24,20 @@ SWidgetSwitcher::FSlot::FSlotArguments SWidgetSwitcher::Slot()
 
 SWidgetSwitcher::FScopedWidgetSlotArguments SWidgetSwitcher::AddSlot(int32 SlotIndex)
 {
+	TWeakPtr<SWidgetSwitcher> WeakSwitcher = SharedThis(this);
 	if (!AllChildren.IsValidIndex(SlotIndex))
 	{
 		// Insert at the end
-		return FScopedWidgetSlotArguments{ MakeUnique<FSlot>(), AllChildren, INDEX_NONE };
+		return FScopedWidgetSlotArguments{ MakeUnique<FSlot>(), AllChildren, INDEX_NONE, [WeakSwitcher](const FSlot*, int32 SlotIndex)
+			{
+				if (TSharedPtr<SWidgetSwitcher> Switcher = WeakSwitcher.Pin())
+				{
+					Switcher->OnSlotAdded(SlotIndex);
+				}
+			} };
 	}
 	else
 	{
-		TWeakPtr<SWidgetSwitcher> WeakSwitcher = SharedThis(this);
 		return FScopedWidgetSlotArguments{ MakeUnique<FSlot>(), AllChildren, SlotIndex, [WeakSwitcher](const FSlot*, int32 SlotIndex)
 			{
 				if (TSharedPtr<SWidgetSwitcher> Switcher = WeakSwitcher.Pin())
@@ -41,6 +47,7 @@ SWidgetSwitcher::FScopedWidgetSlotArguments SWidgetSwitcher::AddSlot(int32 SlotI
 					{
 						Switcher->WidgetIndex = Switcher->WidgetIndex.Get() + 1;
 					}
+					Switcher->OnSlotAdded(SlotIndex);
 				}
 			}};
 	}
@@ -97,12 +104,14 @@ int32 SWidgetSwitcher::RemoveSlot( TSharedRef<SWidget> WidgetToRemove )
 		if (AllChildren[SlotIndex].GetWidget() == WidgetToRemove)
 		{
 			// adjust the active WidgetIndex based on this slot change
+			bool bWasActiveSlot = false;
 			if (!WidgetIndex.IsBound())
 			{
 				int32 ActiveWidgetIndex = WidgetIndex.Get();
 
 				if (SlotIndex == ActiveWidgetIndex)
 				{
+					bWasActiveSlot = true;
 					Invalidate(EInvalidateWidget::ChildOrder);
 				}
 
@@ -113,6 +122,7 @@ int32 SWidgetSwitcher::RemoveSlot( TSharedRef<SWidget> WidgetToRemove )
 			}
 
 			AllChildren.RemoveAt(SlotIndex);
+			OnSlotRemoved(SlotIndex, WidgetToRemove, bWasActiveSlot);
 			return SlotIndex;
 		}
 	}
diff --git a/Engine/Source/Runtime/Slate/Public/Widgets/Layout/SWidgetSwitcher.h b/Engine/Source/Runtime/Slate/Public/Widgets/Layout/SWidgetSwitcher.h
--- a/Engine/Source/Runtime/Slate/Public/Widgets/Layout/SWidgetSwitcher.h
+++ b/Engine/Source/Runtime/Slate/Public/Widgets/Layout/SWidgetSwitcher.h
@@ -140,6 +140,8 @@ protected:
 	SLATE_API virtual void OnArrangeChildren( const FGeometry& AllottedGeometry, FArrangedChildren& ArrangedChildren ) const override;
 	SLATE_API virtual FVector2D ComputeDesiredSize( float ) const override;
 	SLATE_API virtual FChildren* GetChildren( ) override;
+	SLATE_API virtual void OnSlotAdded(int32 AddedIndex) {}
+	SLATE_API virtual void OnSlotRemoved(int32 RemovedIndex, TSharedRef<SWidget> RemovedWidget, bool bWasActiveSlot) {}
 	virtual bool ComputeVolatility() const override { return WidgetIndex.IsBound(); }
 	SLATE_API const FSlot* GetActiveSlot() const;
 
