Pass down the FocusCause when redirecting the focus from the UUserWidget to the DesiredFocusWidget.
Make sure the UCommonButtonInternalBase is not focusable when a DesiredFocusWidget is set in the CommonButtonBase. This will prevent UCommonButtonInternalBase from stealing the focus of the DesiredFocusWidget.
[RN] UUserWidget will now send the FocusCause received to the Desired Focus Widget when set. CommonButtonBase's UCommonButtonInternalBase RootButtonRaw will no longer be focusable when a DesiredFocusWidget is set, preventing it from stealing the focus to the DesiredFocusWidget.

diff --git a/Engine/Plugins/Runtime/CommonUI/Source/CommonUI/Private/CommonButtonBase.cpp b/Engine/Plugins/Runtime/CommonUI/Source/CommonUI/Private/CommonButtonBase.cpp
--- a/Engine/Plugins/Runtime/CommonUI/Source/CommonUI/Private/CommonButtonBase.cpp
+++ b/Engine/Plugins/Runtime/CommonUI/Source/CommonUI/Private/CommonButtonBase.cpp
@@ -423,7 +423,8 @@ bool UCommonButtonBase::Initialize()
 		RootButtonRaw->SetClickMethod(ClickMethod);
 		RootButtonRaw->SetTouchMethod(TouchMethod);
 		RootButtonRaw->SetPressMethod(PressMethod);
-		RootButtonRaw->SetButtonFocusable(IsFocusable());
+		//Force the RootButton to not be focusable if it has a DesiredFocusWidgetName set which was stealing the focus and preventing DesiredFocusWidget from getting the FocusReceived event.
+		RootButtonRaw->SetButtonFocusable(GetDesiredFocusWidgetName().IsNone() && IsFocusable());
 		RootButtonRaw->SetButtonEnabled(bButtonEnabled);
 		RootButtonRaw->SetInteractionEnabled(bInteractionEnabled);
 		RootButton = RootButtonRaw;
diff --git a/Engine/Source/Runtime/UMG/Private/UserWidget.cpp b/Engine/Source/Runtime/UMG/Private/UserWidget.cpp
--- a/Engine/Source/Runtime/UMG/Private/UserWidget.cpp
+++ b/Engine/Source/Runtime/UMG/Private/UserWidget.cpp
@@ -2122,7 +2122,7 @@ FReply UUserWidget::NativeOnFocusReceived( const FGeometry& InGeometry, const FF
 	// Forward focus if Desired Focus is set
 	if (UWidget * WidgetToFocus = DesiredFocusWidget.Resolve(WidgetTree))
 	{
-		return FReply::Handled().SetUserFocus(WidgetToFocus->GetCachedWidget().ToSharedRef());
+		return FReply::Handled().SetUserFocus(WidgetToFocus->GetCachedWidget().ToSharedRef(),InFocusEvent.GetCause());
 	}
 	return Reply;
 }
