Exposing the OnLockClicked and OnLockDoubleClicked events for CommonButtonBase BP accessibility

diff --git a/Engine/Plugins/Runtime/CommonUI/Source/CommonUI/Private/CommonButtonBase.cpp b/Engine/Plugins/Runtime/CommonUI/Source/CommonUI/Private/CommonButtonBase.cpp
--- a/Engine/Plugins/Runtime/CommonUI/Source/CommonUI/Private/CommonButtonBase.cpp
+++ b/Engine/Plugins/Runtime/CommonUI/Source/CommonUI/Private/CommonButtonBase.cpp
@@ -1509,6 +1509,10 @@ void UCommonButtonBase::NativeOnClicked()
 	{
 		BP_OnLockClicked();
 		OnLockClicked().Broadcast();
+		if (OnButtonBaseLockClicked.IsBound())
+		{
+			OnButtonBaseLockClicked.Broadcast(this);
+		}
 	}
 }
 
@@ -1527,6 +1531,10 @@ void UCommonButtonBase::NativeOnDoubleClicked()
 	{
 		BP_OnLockDoubleClicked();
 		OnLockDoubleClicked().Broadcast();
+		if (OnButtonBaseLockDoubleClicked.IsBound())
+		{
+			OnButtonBaseLockDoubleClicked.Broadcast(this);
+		}
 	}
 }
 
diff --git a/Engine/Plugins/Runtime/CommonUI/Source/CommonUI/Private/Groups/CommonButtonGroupBase.cpp b/Engine/Plugins/Runtime/CommonUI/Source/CommonUI/Private/Groups/CommonButtonGroupBase.cpp
--- a/Engine/Plugins/Runtime/CommonUI/Source/CommonUI/Private/Groups/CommonButtonGroupBase.cpp
+++ b/Engine/Plugins/Runtime/CommonUI/Source/CommonUI/Private/Groups/CommonButtonGroupBase.cpp
@@ -210,6 +210,8 @@ void UCommonButtonGroupBase::OnWidgetAdded(UWidget* NewWidget)
 		Button->OnButtonBaseDoubleClicked.AddUniqueDynamic(this, &UCommonButtonGroupBase::OnHandleButtonBaseDoubleClicked);
 		Button->OnButtonBaseHovered.AddUniqueDynamic(this, &UCommonButtonGroupBase::OnButtonBaseHovered);
 		Button->OnButtonBaseUnhovered.AddUniqueDynamic(this, &UCommonButtonGroupBase::OnButtonBaseUnhovered);
+		Button->OnButtonBaseLockClicked.AddUniqueDynamic(this, &UCommonButtonGroupBase::OnHandleButtonBaseLockClicked);
+		Button->OnButtonBaseLockDoubleClicked.AddUniqueDynamic(this, &UCommonButtonGroupBase::OnHandleButtonBaseLockDoubleClicked);
 
 		Buttons.Emplace(Button);
 
@@ -232,6 +234,8 @@ void UCommonButtonGroupBase::OnWidgetRemoved( UWidget* OldWidget )
 		Button->OnButtonBaseDoubleClicked.RemoveDynamic(this, &UCommonButtonGroupBase::OnHandleButtonBaseDoubleClicked);
 		Button->OnButtonBaseHovered.RemoveDynamic(this, &UCommonButtonGroupBase::OnButtonBaseHovered);
 		Button->OnButtonBaseUnhovered.RemoveDynamic(this, &UCommonButtonGroupBase::OnButtonBaseUnhovered);
+		Button->OnButtonBaseLockClicked.RemoveDynamic(this, &UCommonButtonGroupBase::OnHandleButtonBaseLockClicked);
+		Button->OnButtonBaseLockDoubleClicked.RemoveDynamic(this, &UCommonButtonGroupBase::OnHandleButtonBaseLockDoubleClicked);
 
 		Buttons.RemoveAll( [Button]( TWeakObjectPtr<UCommonButtonBase> Entry ) { return Entry == Button || !Entry.IsValid(); } );
 
@@ -371,6 +375,26 @@ void UCommonButtonGroupBase::OnHandleButtonBaseDoubleClicked(UCommonButtonBase*
 	}
 }
 
+void UCommonButtonGroupBase::OnHandleButtonBaseLockClicked(UCommonButtonBase* BaseButton)
+{
+	const int32 ClickedIdx = FindButtonIndex(BaseButton);
+	if (ClickedIdx != INDEX_NONE)
+	{
+		NativeOnButtonBaseLockClicked.Broadcast(BaseButton, ClickedIdx);
+		OnButtonBaseLockClicked.Broadcast(BaseButton, ClickedIdx);
+	}
+}
+
+void UCommonButtonGroupBase::OnHandleButtonBaseLockDoubleClicked(UCommonButtonBase* BaseButton)
+{
+	const int32 ClickedIdx = FindButtonIndex(BaseButton);
+	if (ClickedIdx != INDEX_NONE)
+	{
+		NativeOnButtonBaseLockDoubleClicked.Broadcast(BaseButton, ClickedIdx);
+		OnButtonBaseLockDoubleClicked.Broadcast(BaseButton, ClickedIdx);
+	}
+}
+
 UCommonButtonBase* UCommonButtonGroupBase::GetButtonBaseAtIndex(int32 Index) const
 {
 	if (Buttons.IsValidIndex(Index))
diff --git a/Engine/Plugins/Runtime/CommonUI/Source/CommonUI/Public/CommonButtonBase.h b/Engine/Plugins/Runtime/CommonUI/Source/CommonUI/Public/CommonButtonBase.h
--- a/Engine/Plugins/Runtime/CommonUI/Source/CommonUI/Public/CommonButtonBase.h
+++ b/Engine/Plugins/Runtime/CommonUI/Source/CommonUI/Public/CommonButtonBase.h
@@ -863,6 +863,12 @@ protected:
 
 	UPROPERTY(BlueprintAssignable, Category = "Events", meta = (AllowPrivateAccess = true, DisplayName = "On Unhovered"))
 	FCommonButtonBaseClicked OnButtonBaseUnhovered;
+	
+	UPROPERTY(BlueprintAssignable, Category = "Events", meta = (AllowPrivateAccess = true, DisplayName = "On Lock Clicked"))
+	FCommonButtonBaseClicked OnButtonBaseLockClicked;
+
+	UPROPERTY(BlueprintAssignable, Category = "Events", meta = (AllowPrivateAccess = true, DisplayName = "On Lock Double Clicked"))
+	FCommonButtonBaseClicked OnButtonBaseLockDoubleClicked;
 
 	FUIActionBindingHandle TriggeringBindingHandle;
 	
diff --git a/Engine/Plugins/Runtime/CommonUI/Source/CommonUI/Public/Groups/CommonButtonGroupBase.h b/Engine/Plugins/Runtime/CommonUI/Source/CommonUI/Public/Groups/CommonButtonGroupBase.h
--- a/Engine/Plugins/Runtime/CommonUI/Source/CommonUI/Public/Groups/CommonButtonGroupBase.h
+++ b/Engine/Plugins/Runtime/CommonUI/Source/CommonUI/Public/Groups/CommonButtonGroupBase.h
@@ -106,6 +106,14 @@ public:
 	UPROPERTY(BlueprintAssignable)
 	FOnSelectionCleared OnSelectionCleared;
 	FNativeOnSelectionCleared NativeOnSelectionCleared;
+		
+	UPROPERTY(BlueprintAssignable)
+	FSimpleButtonBaseGroupDelegate OnButtonBaseLockClicked;
+	FNativeSimpleButtonBaseGroupDelegate NativeOnButtonBaseLockClicked;
+
+	UPROPERTY(BlueprintAssignable)
+	FSimpleButtonBaseGroupDelegate OnButtonBaseLockDoubleClicked;
+	FNativeSimpleButtonBaseGroupDelegate NativeOnButtonBaseLockDoubleClicked;
 
 	UFUNCTION(BlueprintCallable, Category = BaseButtonGroup)
 	UCommonButtonBase* GetButtonBaseAtIndex(int32 Index) const;
@@ -142,6 +150,12 @@ protected:
 
 	UFUNCTION()
 	virtual void OnButtonBaseUnhovered(UCommonButtonBase* BaseButton);
+		
+	UFUNCTION()
+	virtual void OnHandleButtonBaseLockClicked(UCommonButtonBase* BaseButton);
+
+	UFUNCTION()
+	virtual void OnHandleButtonBaseLockDoubleClicked(UCommonButtonBase* BaseButton);
 
 protected:
 
