Free widget list in ReleaseSlateResources for Activatable Containers, also clear pool and migrate widgets into inactive list so they can be reused.

#tests PIE Lyra

diff --git a/Engine/Plugins/Runtime/CommonUI/Source/CommonUI/Private/Widgets/CommonActivatableWidgetContainer.cpp b/Engine/Plugins/Runtime/CommonUI/Source/CommonUI/Private/Widgets/CommonActivatableWidgetContainer.cpp
--- a/Engine/Plugins/Runtime/CommonUI/Source/CommonUI/Private/Widgets/CommonActivatableWidgetContainer.cpp
+++ b/Engine/Plugins/Runtime/CommonUI/Source/CommonUI/Private/Widgets/CommonActivatableWidgetContainer.cpp
@@ -122,8 +122,9 @@ void UCommonActivatableWidgetContainerBase::ReleaseSlateResources(bool bReleaseC
 	MyInputGuard.Reset();
 	MySwitcher.Reset();
 	ReleasedWidgets.Empty();
+	WidgetList.Reset();
 
-	GeneratedWidgetsPool.ReleaseAllSlateResources();
+	GeneratedWidgetsPool.ReleaseAll(true);
 }
 
 void UCommonActivatableWidgetContainerBase::OnWidgetRebuilt()
