Remove unnecessary private include path to CommonUI

#rnx

diff --git a/Engine/Plugins/Runtime/CommonUI/Source/CommonUIEditor/CommonUIEditor.Build.cs b/Engine/Plugins/Runtime/CommonUI/Source/CommonUIEditor/CommonUIEditor.Build.cs
--- a/Engine/Plugins/Runtime/CommonUI/Source/CommonUIEditor/CommonUIEditor.Build.cs
+++ b/Engine/Plugins/Runtime/CommonUI/Source/CommonUIEditor/CommonUIEditor.Build.cs
@@ -31,13 +31,6 @@ public class CommonUIEditor : ModuleRules
 			}
         );
 
-		PrivateIncludePaths.AddRange(
-			new string[]
-			{
-				System.IO.Path.Combine(GetModuleDirectory("CommonUI"), "Private"),
-			}
-		);
-
 		PrivateIncludePathModuleNames.AddRange(
 			new string[] {
 				"DataTableEditor",
