PR #11718: CommonUI: Fix not being able to remove widgets from a UCommonActivatableWidgetContainerBase after slate resources have been released

#tests PIE

diff --git a/Engine/Plugins/Runtime/CommonUI/Source/CommonUI/Private/Widgets/CommonActivatableWidgetContainer.cpp b/Engine/Plugins/Runtime/CommonUI/Source/CommonUI/Private/Widgets/CommonActivatableWidgetContainer.cpp
--- a/Engine/Plugins/Runtime/CommonUI/Source/CommonUI/Private/Widgets/CommonActivatableWidgetContainer.cpp
+++ b/Engine/Plugins/Runtime/CommonUI/Source/CommonUI/Private/Widgets/CommonActivatableWidgetContainer.cpp
@@ -81,6 +81,11 @@ void UCommonActivatableWidgetContainerBase::RemoveWidget(UCommonActivatableWidge
 		{
 			ReleaseWidget(CachedWidget.ToSharedRef());
 		}
+		else
+		{
+			GeneratedWidgetsPool.Release(&WidgetToRemove, true);
+			WidgetList.Remove(&WidgetToRemove);
+		}
 	}
 }
 
