PR #11900: Fixed unsafe pointer in CommonActionWidget (CommonUI)

#tests Lyra PIE Gamepad Frontend

diff --git a/Engine/Plugins/Runtime/CommonUI/Source/CommonUI/Public/CommonActionWidget.h b/Engine/Plugins/Runtime/CommonUI/Source/CommonUI/Public/CommonActionWidget.h
--- a/Engine/Plugins/Runtime/CommonUI/Source/CommonUI/Public/CommonActionWidget.h
+++ b/Engine/Plugins/Runtime/CommonUI/Source/CommonUI/Public/CommonActionWidget.h
@@ -137,6 +137,7 @@ protected:
 
 	TSharedPtr<SImage> MyIconRim;
 
+	UPROPERTY()
 	FSlateBrush Icon;
 
 	bool bAlwaysHideOverride = false;
