Disable the top modules that are not currently set to enable AutoRTFM

diff --git a/Engine/Plugins/Animation/RigLogic/Source/RigLogicLib/RigLogicLib.Build.cs b/Engine/Plugins/Animation/RigLogic/Source/RigLogicLib/RigLogicLib.Build.cs
--- a/Engine/Plugins/Animation/RigLogic/Source/RigLogicLib/RigLogicLib.Build.cs
+++ b/Engine/Plugins/Animation/RigLogic/Source/RigLogicLib/RigLogicLib.Build.cs
@@ -52,6 +52,8 @@ namespace UnrealBuildTool.Rules
 
 			PrivateDefinitions.Add("RL_AUTODETECT_SSE=1");
 			PublicDefinitions.Add("RL_BUILD_WITH_ML_EVALUATOR=1");
+
+			bDisableAutoRTFMInstrumentation = true;
 		}
 	}
 }
diff --git a/Engine/Plugins/ChaosCloth/Source/ChaosCloth/ChaosCloth.Build.cs b/Engine/Plugins/ChaosCloth/Source/ChaosCloth/ChaosCloth.Build.cs
--- a/Engine/Plugins/ChaosCloth/Source/ChaosCloth/ChaosCloth.Build.cs
+++ b/Engine/Plugins/ChaosCloth/Source/ChaosCloth/ChaosCloth.Build.cs
@@ -20,5 +20,7 @@ public class ChaosCloth : ModuleRules
 
         SetupModulePhysicsSupport(Target);
 		PrivateDefinitions.Add("CHAOS_INCLUDE_LEVEL_1=1");
+
+		bDisableAutoRTFMInstrumentation = true;
 	}
 }
diff --git a/Engine/Plugins/Enterprise/GLTFExporter/Source/GLTFExporter/GLTFExporter.Build.cs b/Engine/Plugins/Enterprise/GLTFExporter/Source/GLTFExporter/GLTFExporter.Build.cs
--- a/Engine/Plugins/Enterprise/GLTFExporter/Source/GLTFExporter/GLTFExporter.Build.cs
+++ b/Engine/Plugins/Enterprise/GLTFExporter/Source/GLTFExporter/GLTFExporter.Build.cs
@@ -83,5 +83,7 @@ public class GLTFExporter : ModuleRules
 				}
 			);
 		}
+
+		bDisableAutoRTFMInstrumentation = true;
 	}
 }
diff --git a/Engine/Plugins/Experimental/LearningAgents/Source/Learning/Learning.Build.cs b/Engine/Plugins/Experimental/LearningAgents/Source/Learning/Learning.Build.cs
--- a/Engine/Plugins/Experimental/LearningAgents/Source/Learning/Learning.Build.cs
+++ b/Engine/Plugins/Experimental/LearningAgents/Source/Learning/Learning.Build.cs
@@ -53,5 +53,7 @@ public class Learning : ModuleRules
 				// ... add any modules that your module loads dynamically here ...
 			}
 			);
+
+		bDisableAutoRTFMInstrumentation = true;
 	}
 }
diff --git a/Engine/Plugins/Experimental/Mutable/Source/MutableRuntime/MutableRuntime.Build.cs b/Engine/Plugins/Experimental/Mutable/Source/MutableRuntime/MutableRuntime.Build.cs
--- a/Engine/Plugins/Experimental/Mutable/Source/MutableRuntime/MutableRuntime.Build.cs
+++ b/Engine/Plugins/Experimental/Mutable/Source/MutableRuntime/MutableRuntime.Build.cs
@@ -27,6 +27,8 @@ namespace UnrealBuildTool.Rules
 				"TraceLog"
 				}
 			);
+
+			bDisableAutoRTFMInstrumentation = true;
 		}
 	}
 }
diff --git a/Engine/Plugins/Experimental/PlanarCutPlugin/Source/PlanarCut/PlanarCut.Build.cs b/Engine/Plugins/Experimental/PlanarCutPlugin/Source/PlanarCut/PlanarCut.Build.cs
--- a/Engine/Plugins/Experimental/PlanarCutPlugin/Source/PlanarCut/PlanarCut.Build.cs
+++ b/Engine/Plugins/Experimental/PlanarCutPlugin/Source/PlanarCut/PlanarCut.Build.cs
@@ -29,6 +29,8 @@ namespace UnrealBuildTool.Rules
 					"MeshDescription",
                 }
                 );
+
+			bDisableAutoRTFMInstrumentation = true;
 		}
 	}
 }
diff --git a/Engine/Plugins/FX/Niagara/Source/Niagara/Niagara.Build.cs b/Engine/Plugins/FX/Niagara/Source/Niagara/Niagara.Build.cs
--- a/Engine/Plugins/FX/Niagara/Source/Niagara/Niagara.Build.cs
+++ b/Engine/Plugins/FX/Niagara/Source/Niagara/Niagara.Build.cs
@@ -69,5 +69,7 @@ public class Niagara : ModuleRules
                 "VECTORVM_SUPPORTS_EXPERIMENTAL=1",
                 "VECTORVM_SUPPORTS_LEGACY=1"
             });
+
+		bDisableAutoRTFMInstrumentation = true;
 	}
 }
diff --git a/Engine/Plugins/Interchange/Runtime/Source/Pipelines/InterchangePipelines.Build.cs b/Engine/Plugins/Interchange/Runtime/Source/Pipelines/InterchangePipelines.Build.cs
--- a/Engine/Plugins/Interchange/Runtime/Source/Pipelines/InterchangePipelines.Build.cs
+++ b/Engine/Plugins/Interchange/Runtime/Source/Pipelines/InterchangePipelines.Build.cs
@@ -46,6 +46,8 @@ namespace UnrealBuildTool.Rules
 					}
 				);
 			}
+
+			bDisableAutoRTFMInstrumentation = true;
 		}
 	}
 }
diff --git a/Engine/Plugins/Online/OnlineFramework/Source/Party/Party.Build.cs b/Engine/Plugins/Online/OnlineFramework/Source/Party/Party.Build.cs
--- a/Engine/Plugins/Online/OnlineFramework/Source/Party/Party.Build.cs
+++ b/Engine/Plugins/Online/OnlineFramework/Source/Party/Party.Build.cs
@@ -34,6 +34,8 @@ public class Party : ModuleRules
 		else if (Target.Platform == UnrealTargetPlatform.Mac)
 		{
 		}
+
+		bDisableAutoRTFMInstrumentation = true;
 	}
 
 	protected virtual bool bUsesPSNSessions { get { return false; } }
diff --git a/Engine/Plugins/Online/OnlineServices/Source/OnlineServicesCommon/OnlineServicesCommon.Build.cs b/Engine/Plugins/Online/OnlineServices/Source/OnlineServicesCommon/OnlineServicesCommon.Build.cs
--- a/Engine/Plugins/Online/OnlineServices/Source/OnlineServicesCommon/OnlineServicesCommon.Build.cs
+++ b/Engine/Plugins/Online/OnlineServices/Source/OnlineServicesCommon/OnlineServicesCommon.Build.cs
@@ -24,5 +24,7 @@ public class OnlineServicesCommon : ModuleRules
 				"CoreUObject"		// CoreUObject temporary dependency
 			}
 		);
+
+		bDisableAutoRTFMInstrumentation = true;
 	}
 }
diff --git a/Engine/Plugins/Online/OnlineServicesEOS/Source/OnlineServicesEOS.Build.cs b/Engine/Plugins/Online/OnlineServicesEOS/Source/OnlineServicesEOS.Build.cs
--- a/Engine/Plugins/Online/OnlineServicesEOS/Source/OnlineServicesEOS.Build.cs
+++ b/Engine/Plugins/Online/OnlineServicesEOS/Source/OnlineServicesEOS.Build.cs
@@ -24,5 +24,7 @@ public class OnlineServicesEOS : ModuleRules
 				"EOSShared",
 			}
 		);
+
+		bDisableAutoRTFMInstrumentation = true;
 	}
 }
diff --git a/Engine/Plugins/Online/OnlineServicesEOSGS/Source/OnlineServicesEOSGS.Build.cs b/Engine/Plugins/Online/OnlineServicesEOSGS/Source/OnlineServicesEOSGS.Build.cs
--- a/Engine/Plugins/Online/OnlineServicesEOSGS/Source/OnlineServicesEOSGS.Build.cs
+++ b/Engine/Plugins/Online/OnlineServicesEOSGS/Source/OnlineServicesEOSGS.Build.cs
@@ -42,5 +42,7 @@ public class OnlineServicesEOSGS : ModuleRules
 			PublicWeakFrameworks.Add("AuthenticationServices");
 			PrivateDependencyModuleNames.Add("ApplicationCore");
 		}
+
+		bDisableAutoRTFMInstrumentation = true;
 	}
 }
diff --git a/Engine/Plugins/Online/OnlineServicesOSSAdapter/Source/OnlineServicesOSSAdapter.Build.cs b/Engine/Plugins/Online/OnlineServicesOSSAdapter/Source/OnlineServicesOSSAdapter.Build.cs
--- a/Engine/Plugins/Online/OnlineServicesOSSAdapter/Source/OnlineServicesOSSAdapter.Build.cs
+++ b/Engine/Plugins/Online/OnlineServicesOSSAdapter/Source/OnlineServicesOSSAdapter.Build.cs
@@ -22,5 +22,7 @@ public class OnlineServicesOSSAdapter : ModuleRules
 				"CoreOnline",
 			}
 		);
+
+		bDisableAutoRTFMInstrumentation = true;
 	}
 }
diff --git a/Engine/Plugins/PCG/Source/PCG/PCG.Build.cs b/Engine/Plugins/PCG/Source/PCG/PCG.Build.cs
--- a/Engine/Plugins/PCG/Source/PCG/PCG.Build.cs
+++ b/Engine/Plugins/PCG/Source/PCG/PCG.Build.cs
@@ -43,6 +43,8 @@ namespace UnrealBuildTool.Rules
 					}
 				);
 			}
+
+			bDisableAutoRTFMInstrumentation = true;
 		}
 	}
 }
diff --git a/Engine/Plugins/Runtime/AudioModulation/Source/AudioModulation/AudioModulation.Build.cs b/Engine/Plugins/Runtime/AudioModulation/Source/AudioModulation/AudioModulation.Build.cs
--- a/Engine/Plugins/Runtime/AudioModulation/Source/AudioModulation/AudioModulation.Build.cs
+++ b/Engine/Plugins/Runtime/AudioModulation/Source/AudioModulation/AudioModulation.Build.cs
@@ -64,6 +64,8 @@ namespace UnrealBuildTool.Rules
 			{
 				PublicDefinitions.Add("WITH_AUDIOMODULATION_METASOUND_SUPPORT=0");
 			}
+
+			bDisableAutoRTFMInstrumentation = true;
 		}
 	}
-}
\ No newline at end of file
+}
diff --git a/Engine/Plugins/Runtime/AudioWidgets/Source/AudioWidgets/AudioWidgets.Build.cs b/Engine/Plugins/Runtime/AudioWidgets/Source/AudioWidgets/AudioWidgets.Build.cs
--- a/Engine/Plugins/Runtime/AudioWidgets/Source/AudioWidgets/AudioWidgets.Build.cs
+++ b/Engine/Plugins/Runtime/AudioWidgets/Source/AudioWidgets/AudioWidgets.Build.cs
@@ -44,5 +44,7 @@ public class AudioWidgets : ModuleRules
 				}
 			);
 		}
+
+		bDisableAutoRTFMInstrumentation = true;
 	}
 }
diff --git a/Engine/Plugins/Runtime/CommonUI/Source/CommonUI/CommonUI.Build.cs b/Engine/Plugins/Runtime/CommonUI/Source/CommonUI/CommonUI.Build.cs
--- a/Engine/Plugins/Runtime/CommonUI/Source/CommonUI/CommonUI.Build.cs
+++ b/Engine/Plugins/Runtime/CommonUI/Source/CommonUI/CommonUI.Build.cs
@@ -47,6 +47,8 @@ public class CommonUI : ModuleRules
         }
 		PrivateDefinitions.Add("UE_COMMONUI_PLATFORM_KBM_REQUIRES_ATTACHED_MOUSE=" + (bPlatformKBMRequiresAttachedMouse ? "1" : "0"));
 		PrivateDefinitions.Add("UE_COMMONUI_PLATFORM_REQUIRES_CURSOR_HIDDEN_FOR_TOUCH=" + (bPlatformRequiresCursorHiddenForTouch ? "1" : "0"));
+
+		bDisableAutoRTFMInstrumentation = true;
 	}
 
 	protected virtual bool bPlatformKBMRequiresAttachedMouse { get { return false; } }
diff --git a/Engine/Plugins/Runtime/GeometryProcessing/Source/DynamicMesh/DynamicMesh.Build.cs b/Engine/Plugins/Runtime/GeometryProcessing/Source/DynamicMesh/DynamicMesh.Build.cs
--- a/Engine/Plugins/Runtime/GeometryProcessing/Source/DynamicMesh/DynamicMesh.Build.cs
+++ b/Engine/Plugins/Runtime/GeometryProcessing/Source/DynamicMesh/DynamicMesh.Build.cs
@@ -28,5 +28,7 @@ public class DynamicMesh : ModuleRules
 			}
 		);
 		UnsafeTypeCastWarningLevel = WarningLevel.Error;
+
+		bDisableAutoRTFMInstrumentation = true;
 	}
 }
diff --git a/Engine/Plugins/Runtime/GeometryProcessing/Source/GeometryAlgorithms/GeometryAlgorithms.Build.cs b/Engine/Plugins/Runtime/GeometryProcessing/Source/GeometryAlgorithms/GeometryAlgorithms.Build.cs
--- a/Engine/Plugins/Runtime/GeometryProcessing/Source/GeometryAlgorithms/GeometryAlgorithms.Build.cs
+++ b/Engine/Plugins/Runtime/GeometryProcessing/Source/GeometryAlgorithms/GeometryAlgorithms.Build.cs
@@ -38,5 +38,7 @@ public class GeometryAlgorithms : ModuleRules
 				"IntelTBB"
 			);
 		}
+
+		bDisableAutoRTFMInstrumentation = true;
 	}
 }
diff --git a/Engine/Plugins/Runtime/GeometryScripting/Source/GeometryScriptingCore/GeometryScriptingCore.Build.cs b/Engine/Plugins/Runtime/GeometryScripting/Source/GeometryScriptingCore/GeometryScriptingCore.Build.cs
--- a/Engine/Plugins/Runtime/GeometryScripting/Source/GeometryScriptingCore/GeometryScriptingCore.Build.cs
+++ b/Engine/Plugins/Runtime/GeometryScripting/Source/GeometryScriptingCore/GeometryScriptingCore.Build.cs
@@ -69,5 +69,6 @@ public class GeometryScriptingCore : ModuleRules
 			PrivateDependencyModuleNames.Add("UnrealEd");
 		}
 
+		bDisableAutoRTFMInstrumentation = true;
 	}
 }
diff --git a/Engine/Plugins/Runtime/Harmonix/Source/HarmonixMetasound/HarmonixMetasound.Build.cs b/Engine/Plugins/Runtime/Harmonix/Source/HarmonixMetasound/HarmonixMetasound.Build.cs
--- a/Engine/Plugins/Runtime/Harmonix/Source/HarmonixMetasound/HarmonixMetasound.Build.cs
+++ b/Engine/Plugins/Runtime/Harmonix/Source/HarmonixMetasound/HarmonixMetasound.Build.cs
@@ -41,5 +41,7 @@ public class HarmonixMetasound : ModuleRules
 			PrivateDependencyModuleNames.Add("AssetRegistry");
 			PrivateDependencyModuleNames.Add("UnrealEd");
 		}
+
+		bDisableAutoRTFMInstrumentation = true;
 	}
 }
diff --git a/Engine/Plugins/Runtime/MeshModelingToolset/Source/MeshModelingTools/MeshModelingTools.Build.cs b/Engine/Plugins/Runtime/MeshModelingToolset/Source/MeshModelingTools/MeshModelingTools.Build.cs
--- a/Engine/Plugins/Runtime/MeshModelingToolset/Source/MeshModelingTools/MeshModelingTools.Build.cs
+++ b/Engine/Plugins/Runtime/MeshModelingToolset/Source/MeshModelingTools/MeshModelingTools.Build.cs
@@ -35,5 +35,7 @@ public class MeshModelingTools : ModuleRules
 				// ... add private dependencies that you statically link with here ...
 			}
 		);
+
+		bDisableAutoRTFMInstrumentation = true;
 	}
-}
\ No newline at end of file
+}
diff --git a/Engine/Plugins/Runtime/Metasound/Source/MetasoundEngine/MetasoundEngine.Build.cs b/Engine/Plugins/Runtime/Metasound/Source/MetasoundEngine/MetasoundEngine.Build.cs
--- a/Engine/Plugins/Runtime/Metasound/Source/MetasoundEngine/MetasoundEngine.Build.cs
+++ b/Engine/Plugins/Runtime/Metasound/Source/MetasoundEngine/MetasoundEngine.Build.cs
@@ -44,6 +44,8 @@ namespace UnrealBuildTool.Rules
 					"WaveTable"
 				}
 			);
+
+			bDisableAutoRTFMInstrumentation = true;
 		}
 	}
 }
diff --git a/Engine/Plugins/Runtime/Metasound/Source/MetasoundFrontend/MetasoundFrontend.Build.cs b/Engine/Plugins/Runtime/Metasound/Source/MetasoundFrontend/MetasoundFrontend.Build.cs
--- a/Engine/Plugins/Runtime/Metasound/Source/MetasoundFrontend/MetasoundFrontend.Build.cs
+++ b/Engine/Plugins/Runtime/Metasound/Source/MetasoundFrontend/MetasoundFrontend.Build.cs
@@ -25,6 +25,8 @@ namespace UnrealBuildTool.Rules
 			);
 
 			PublicDefinitions.Add("WITH_METASOUND_FRONTEND=1");
+
+			bDisableAutoRTFMInstrumentation = true;
 		}
 	}
 }
diff --git a/Engine/Plugins/Runtime/Metasound/Source/MetasoundStandardNodes/MetasoundStandardNodes.Build.cs b/Engine/Plugins/Runtime/Metasound/Source/MetasoundStandardNodes/MetasoundStandardNodes.Build.cs
--- a/Engine/Plugins/Runtime/Metasound/Source/MetasoundStandardNodes/MetasoundStandardNodes.Build.cs
+++ b/Engine/Plugins/Runtime/Metasound/Source/MetasoundStandardNodes/MetasoundStandardNodes.Build.cs
@@ -31,6 +31,8 @@ namespace UnrealBuildTool.Rules
 			);
 
 			NumIncludedBytesPerUnityCPPOverride = 120 * 1024;
+
+			bDisableAutoRTFMInstrumentation = true;
 		}
 	}
 }
diff --git a/Engine/Plugins/Runtime/RigVM/Source/RigVM/RigVM.Build.cs b/Engine/Plugins/Runtime/RigVM/Source/RigVM/RigVM.Build.cs
--- a/Engine/Plugins/Runtime/RigVM/Source/RigVM/RigVM.Build.cs
+++ b/Engine/Plugins/Runtime/RigVM/Source/RigVM/RigVM.Build.cs
@@ -33,5 +33,7 @@ public class RigVM : ModuleRules
 				}
             );
         }
+
+		bDisableAutoRTFMInstrumentation = true;
     }
 }
diff --git a/Engine/Plugins/Runtime/Synthesis/Source/Synthesis/Synthesis.Build.cs b/Engine/Plugins/Runtime/Synthesis/Source/Synthesis/Synthesis.Build.cs
--- a/Engine/Plugins/Runtime/Synthesis/Source/Synthesis/Synthesis.Build.cs
+++ b/Engine/Plugins/Runtime/Synthesis/Source/Synthesis/Synthesis.Build.cs
@@ -31,6 +31,8 @@ namespace UnrealBuildTool.Rules
 			{
 				PrivateDependencyModuleNames.Add("AudioSynesthesiaCore");
 			}
+
+			bDisableAutoRTFMInstrumentation = true;
 		}
 	}
-}
\ No newline at end of file
+}
diff --git a/Engine/Source/Developer/MessageLog/MessageLog.Build.cs b/Engine/Source/Developer/MessageLog/MessageLog.Build.cs
--- a/Engine/Source/Developer/MessageLog/MessageLog.Build.cs
+++ b/Engine/Source/Developer/MessageLog/MessageLog.Build.cs
@@ -32,5 +32,7 @@ public class MessageLog : ModuleRules
 		{
 			PrecompileForTargets = PrecompileTargetsType.Any;
 		}
+
+		bDisableAutoRTFMInstrumentation = true;
 	}
 }
diff --git a/Engine/Source/Developer/SlateReflector/SlateReflector.Build.cs b/Engine/Source/Developer/SlateReflector/SlateReflector.Build.cs
--- a/Engine/Source/Developer/SlateReflector/SlateReflector.Build.cs
+++ b/Engine/Source/Developer/SlateReflector/SlateReflector.Build.cs
@@ -95,5 +95,7 @@ public class SlateReflector : ModuleRules
 		}
 
 		UnsafeTypeCastWarningLevel = WarningLevel.Error;
+
+		bDisableAutoRTFMInstrumentation = true;
 	}
 }
diff --git a/Engine/Source/Developer/ToolWidgets/ToolWidgets.Build.cs b/Engine/Source/Developer/ToolWidgets/ToolWidgets.Build.cs
--- a/Engine/Source/Developer/ToolWidgets/ToolWidgets.Build.cs
+++ b/Engine/Source/Developer/ToolWidgets/ToolWidgets.Build.cs
@@ -27,5 +27,7 @@ public class ToolWidgets : ModuleRules
 				"ApplicationCore"
 			}
 		);
+
+		bDisableAutoRTFMInstrumentation = true;
 	}
 }
diff --git a/Engine/Source/Runtime/AdvancedWidgets/AdvancedWidgets.Build.cs b/Engine/Source/Runtime/AdvancedWidgets/AdvancedWidgets.Build.cs
--- a/Engine/Source/Runtime/AdvancedWidgets/AdvancedWidgets.Build.cs
+++ b/Engine/Source/Runtime/AdvancedWidgets/AdvancedWidgets.Build.cs
@@ -28,5 +28,7 @@ public class AdvancedWidgets : ModuleRules
 		{
 			PrivateDependencyModuleNames.Add("UnrealEd");
 		}
+
+		bDisableAutoRTFMInstrumentation = true;
 	}
 }
diff --git a/Engine/Source/Runtime/AnimGraphRuntime/AnimGraphRuntime.Build.cs b/Engine/Source/Runtime/AnimGraphRuntime/AnimGraphRuntime.Build.cs
--- a/Engine/Source/Runtime/AnimGraphRuntime/AnimGraphRuntime.Build.cs
+++ b/Engine/Source/Runtime/AnimGraphRuntime/AnimGraphRuntime.Build.cs
@@ -34,5 +34,7 @@ public class AnimGraphRuntime : ModuleRules
 				"GeometryCollectionEngine",
             }
         );
+
+		bDisableAutoRTFMInstrumentation = true;
     }
 }
diff --git a/Engine/Source/Runtime/AppFramework/AppFramework.Build.cs b/Engine/Source/Runtime/AppFramework/AppFramework.Build.cs
--- a/Engine/Source/Runtime/AppFramework/AppFramework.Build.cs
+++ b/Engine/Source/Runtime/AppFramework/AppFramework.Build.cs
@@ -38,5 +38,7 @@ public class AppFramework : ModuleRules
 				}
 			);
 		}
+
+		bDisableAutoRTFMInstrumentation = true;
 	}
 }
diff --git a/Engine/Source/Runtime/Datasmith/DatasmithCore/DatasmithCore.Build.cs b/Engine/Source/Runtime/Datasmith/DatasmithCore/DatasmithCore.Build.cs
--- a/Engine/Source/Runtime/Datasmith/DatasmithCore/DatasmithCore.Build.cs
+++ b/Engine/Source/Runtime/Datasmith/DatasmithCore/DatasmithCore.Build.cs
@@ -27,6 +27,8 @@ namespace UnrealBuildTool.Rules
 					"XmlParser",
 				}
 			);
+
+			bDisableAutoRTFMInstrumentation = true;
 		}
 	}
-}
\ No newline at end of file
+}
diff --git a/Engine/Source/Runtime/Experimental/GeometryCollectionEngine/GeometryCollectionEngine.Build.cs b/Engine/Source/Runtime/Experimental/GeometryCollectionEngine/GeometryCollectionEngine.Build.cs
--- a/Engine/Source/Runtime/Experimental/GeometryCollectionEngine/GeometryCollectionEngine.Build.cs
+++ b/Engine/Source/Runtime/Experimental/GeometryCollectionEngine/GeometryCollectionEngine.Build.cs
@@ -48,6 +48,8 @@ namespace UnrealBuildTool.Rules
 			}
 
 			PrivateDefinitions.Add("CHAOS_INCLUDE_LEVEL_1=1");
+
+			bDisableAutoRTFMInstrumentation = true;
 		}
 	}
 }
diff --git a/Engine/Source/Runtime/GeometryCore/GeometryCore.Build.cs b/Engine/Source/Runtime/GeometryCore/GeometryCore.Build.cs
--- a/Engine/Source/Runtime/GeometryCore/GeometryCore.Build.cs
+++ b/Engine/Source/Runtime/GeometryCore/GeometryCore.Build.cs
@@ -22,5 +22,7 @@ public class GeometryCore : ModuleRules
 		);
 
 		UnsafeTypeCastWarningLevel = WarningLevel.Error;
+
+		bDisableAutoRTFMInstrumentation = true;
 	}
 }
diff --git a/Engine/Source/Runtime/ImageWrapper/ImageWrapper.Build.cs b/Engine/Source/Runtime/ImageWrapper/ImageWrapper.Build.cs
--- a/Engine/Source/Runtime/ImageWrapper/ImageWrapper.Build.cs
+++ b/Engine/Source/Runtime/ImageWrapper/ImageWrapper.Build.cs
@@ -78,5 +78,7 @@ public class ImageWrapper : ModuleRules
 
 		// Enable exceptions to allow error handling
 		bEnableExceptions = true;
+
+		bDisableAutoRTFMInstrumentation = true;
 	}
 }
diff --git a/Engine/Source/Runtime/Slate/Slate.Build.cs b/Engine/Source/Runtime/Slate/Slate.Build.cs
--- a/Engine/Source/Runtime/Slate/Slate.Build.cs
+++ b/Engine/Source/Runtime/Slate/Slate.Build.cs
@@ -60,5 +60,7 @@ public class Slate : ModuleRules
         }
 
 		PrivateIncludePaths.Add(System.IO.Path.Combine(EngineDirectory,"Source/ThirdParty/AHEasing/AHEasing-1.3.2"));
+
+		bDisableAutoRTFMInstrumentation = true;
 	}
 }
diff --git a/Engine/Source/Runtime/SlateCore/SlateCore.Build.cs b/Engine/Source/Runtime/SlateCore/SlateCore.Build.cs
--- a/Engine/Source/Runtime/SlateCore/SlateCore.Build.cs
+++ b/Engine/Source/Runtime/SlateCore/SlateCore.Build.cs
@@ -53,5 +53,7 @@ public class SlateCore : ModuleRules
 		{
 			AddEngineThirdPartyPrivateStaticDependencies(Target, "XInput");
 		}
+
+		bDisableAutoRTFMInstrumentation = true;
 	}
 }
