Add hold state query methods to CommonButtonBase.

#rnx

diff --git a/Engine/Plugins/Runtime/CommonUI/Source/CommonUI/Private/CommonButtonBase.cpp b/Engine/Plugins/Runtime/CommonUI/Source/CommonUI/Private/CommonButtonBase.cpp
--- a/Engine/Plugins/Runtime/CommonUI/Source/CommonUI/Private/CommonButtonBase.cpp
+++ b/Engine/Plugins/Runtime/CommonUI/Source/CommonUI/Private/CommonButtonBase.cpp
@@ -1908,6 +1908,15 @@ void UCommonButtonBase::DisableButton()
 	}
 }
 
+void UCommonButtonBase::SetRequiresHold(bool bInRequiresHold)
+{
+	bRequiresHold = bInRequiresHold;
+	if (const UCommonInputSubsystem* CommonInputSubsystem = GetInputSubsystem())
+	{
+		UpdateHoldData(CommonInputSubsystem->GetCurrentInputType());
+	}
+}
+
 void UCommonButtonBase::SetIsFocusable(bool bInIsFocusable)
 {
 	UUserWidget::SetIsFocusable(bInIsFocusable);
diff --git a/Engine/Plugins/Runtime/CommonUI/Source/CommonUI/Public/CommonButtonBase.h b/Engine/Plugins/Runtime/CommonUI/Source/CommonUI/Public/CommonButtonBase.h
--- a/Engine/Plugins/Runtime/CommonUI/Source/CommonUI/Public/CommonButtonBase.h
+++ b/Engine/Plugins/Runtime/CommonUI/Source/CommonUI/Public/CommonButtonBase.h
@@ -432,6 +432,18 @@ public:
 	UFUNCTION(BlueprintCallable, Category = "Common Button|Getters")
 	UInputAction* GetEnhancedInputAction() const;
 
+	/** Returns true if this button has a hold behavior, even if the triggering action is not holdable. */
+	UFUNCTION(BlueprintCallable, BlueprintPure, Category = "Common Button|Getters")
+	bool GetRequiresHold() const { return bRequiresHold; }
+
+	/** Change whether this button should have a hold behavior even if the triggering action is not holdable. */
+	UFUNCTION(BlueprintCallable, Category = "Common Button|Setters")
+	void SetRequiresHold(bool bInRequiresHold);
+
+	/** Returns required hold time for performing a triggering action. */
+	UFUNCTION(BlueprintCallable, BlueprintPure, Category = "Common Button|Getters")
+	float GetRequiredHoldTime() const { return HoldTime; };
+
 	/** Updates the bIsFocusable flag */
 	UFUNCTION(BlueprintCallable, Category = "Common Button|Getters")
 	void SetIsFocusable(bool bInIsFocusable);
@@ -517,7 +529,7 @@ protected:
 	/** Called via delegate when the input method changes */
 	UFUNCTION()
 	virtual void OnInputMethodChanged(ECommonInputType CurrentInputType);
-	
+
 	/** If HoldData is valid, assigns its values to Keyboard and Mouse, Gamepad and Touch, based off the Current Input Type. */
     UFUNCTION()
 	virtual void UpdateHoldData(ECommonInputType CurrentInputType);
