Updated CommonActionWidget to add new function for getting the Dynamic Material Instance of the icon

This change came as a result of a suggestion made by [REDACTED_HANDLE] on this slack thread: https://epicgames.slack.com/archives/C044SQFRJ77/p1715606821365059

We willl need to test this change against in progress blueprint changes made be [REDACTED_HANDLE] to see if it solves what he's looking to do. If this change is not successful in accomplishing what Alex needs, I will likely need to investigate having these stylized icons managed by CommonUI instead in a more data-driven way suggested by [REDACTED_HANDLE]: https://epicgames.slack.com/archives/C044QE26U4C/p1717103172767049

#tests
- Ran editor locally successfully

#rnx

diff --git a/Engine/Plugins/Runtime/CommonUI/Source/CommonUI/Private/CommonActionWidget.cpp b/Engine/Plugins/Runtime/CommonUI/Source/CommonUI/Private/CommonActionWidget.cpp
--- a/Engine/Plugins/Runtime/CommonUI/Source/CommonUI/Private/CommonActionWidget.cpp
+++ b/Engine/Plugins/Runtime/CommonUI/Source/CommonUI/Private/CommonActionWidget.cpp
@@ -197,6 +197,28 @@ FText UCommonActionWidget::GetDisplayText() const
 	return FText();
 }
 
+UMaterialInstanceDynamic* UCommonActionWidget::GetIconDynamicMaterial()
+{
+	if (UMaterialInterface* Material = Cast<UMaterialInterface>(Icon.GetResourceObject()))
+	{
+		UMaterialInstanceDynamic* DynamicMaterial = Cast<UMaterialInstanceDynamic>(Material);
+
+		if (!DynamicMaterial)
+		{
+			DynamicMaterial = UMaterialInstanceDynamic::Create(Material, this);
+			Icon.SetResourceObject(DynamicMaterial);
+
+			if (MyIcon.IsValid())
+			{
+				MyIcon->InvalidateImage();
+			}
+		}
+		return DynamicMaterial;
+	}
+
+	return nullptr;
+}
+
 bool UCommonActionWidget::IsHeldAction() const
 {
 	if (EnhancedInputAction && CommonUI::IsEnhancedInputSupportEnabled())
diff --git a/Engine/Plugins/Runtime/CommonUI/Source/CommonUI/Public/CommonActionWidget.h b/Engine/Plugins/Runtime/CommonUI/Source/CommonUI/Public/CommonActionWidget.h
--- a/Engine/Plugins/Runtime/CommonUI/Source/CommonUI/Public/CommonActionWidget.h
+++ b/Engine/Plugins/Runtime/CommonUI/Source/CommonUI/Public/CommonActionWidget.h
@@ -42,6 +42,9 @@ public:
 	UFUNCTION(BlueprintCallable, Category = CommonActionWidget)
 	FText GetDisplayText() const;
 
+	UFUNCTION(BlueprintCallable, Category = CommonActionWidget)
+	UMaterialInstanceDynamic* GetIconDynamicMaterial();
+
 	UFUNCTION(BlueprintCallable, Category = CommonActionWidget)
 	void SetEnhancedInputAction(UInputAction* InInputAction);
 
