Added a CVar "CommonInput.ResetCursorPositionOnTouchInput" (Disabled by default) that reset the virtual cursor position on input method changed to touch.

diff --git a/Engine/Plugins/Runtime/CommonUI/Source/CommonInput/Private/CommonInputSubsystem.cpp b/Engine/Plugins/Runtime/CommonUI/Source/CommonInput/Private/CommonInputSubsystem.cpp
--- a/Engine/Plugins/Runtime/CommonUI/Source/CommonInput/Private/CommonInputSubsystem.cpp
+++ b/Engine/Plugins/Runtime/CommonUI/Source/CommonInput/Private/CommonInputSubsystem.cpp
@@ -16,6 +16,7 @@
 #include "Framework/Application/SlateUser.h"
 #include "Misc/ConfigCacheIni.h"
 #include "HAL/PlatformStackWalk.h"
+#include "Slate/SceneViewport.h"
 #include "Widgets/SViewport.h"
 #include "CommonInputSettings.h"
 #include "ICommonInputModule.h"
@@ -54,6 +55,14 @@ static const FAutoConsoleVariableRef CVarInputEnableGamepadPlatformCursor
 	TEXT("Should the cursor be allowed to be used during gamepad input")
 );
 
+bool bResetCursorPositionOnTouchInput = false;
+static const FAutoConsoleVariableRef CVarResetCursorPositionOnTouchInput
+(
+	TEXT("CommonInput.ResetCursorPositionOnTouchInput"),
+	bResetCursorPositionOnTouchInput,
+	TEXT("Should the cursor position be reset when the user switches to touch inputs.")
+);
+
 UCommonInputSubsystem* UCommonInputSubsystem::Get(const ULocalPlayer* LocalPlayer)
 {
 	return LocalPlayer ? LocalPlayer->GetSubsystem<UCommonInputSubsystem>() : nullptr;
@@ -218,6 +227,16 @@ void UCommonInputSubsystem::BroadcastInputMethodChanged()
 			OnInputMethodChangedNative.Broadcast(CurrentInputType);
 			OnInputMethodChanged.Broadcast(CurrentInputType);
 			LastInputMethodChangeTime = FPlatformTime::Seconds();
+			
+			if (bResetCursorPositionOnTouchInput && CurrentInputType == ECommonInputType::Touch)
+			{
+				UGameViewportClient* GameViewport = World->GetGameViewport();
+				if (GameViewport && GameViewport->GetWindow().IsValid() && GameViewport->Viewport)
+				{
+					const FVector2D TopLeftPos = GameViewport->Viewport->ViewportToVirtualDesktopPixel(FVector2D(0.025f, 0.025f));
+					SetCursorPosition(TopLeftPos, false);
+				}
+			}
 		}
 	}
 }
