Added new OnInputIconUpdated event to CommonActionWidget

This event is necessary for the blueprint to listen to icon changes in order to set the dark mode on every icon update

#tests
- Ran editor locally successfully

#rnx

diff --git a/Engine/Plugins/Runtime/CommonUI/Source/CommonUI/Private/CommonActionWidget.cpp b/Engine/Plugins/Runtime/CommonUI/Source/CommonUI/Private/CommonActionWidget.cpp
--- a/Engine/Plugins/Runtime/CommonUI/Source/CommonUI/Private/CommonActionWidget.cpp
+++ b/Engine/Plugins/Runtime/CommonUI/Source/CommonUI/Private/CommonActionWidget.cpp
@@ -333,6 +333,7 @@ void UCommonActionWidget::UpdateActionWidget()
 				else if (MyIcon.IsValid())
 				{
 					MyIcon->SetImage(&Icon);
+					OnInputIconUpdated.Broadcast();
 
 					if (GetVisibility() != ESlateVisibility::Collapsed)
 					{
diff --git a/Engine/Plugins/Runtime/CommonUI/Source/CommonUI/Public/CommonActionWidget.h b/Engine/Plugins/Runtime/CommonUI/Source/CommonUI/Public/CommonActionWidget.h
--- a/Engine/Plugins/Runtime/CommonUI/Source/CommonUI/Public/CommonActionWidget.h
+++ b/Engine/Plugins/Runtime/CommonUI/Source/CommonUI/Public/CommonActionWidget.h
@@ -66,6 +66,10 @@ public:
 	UPROPERTY(BlueprintAssignable, Category = CommonActionWidget)
 	FOnInputMethodChanged OnInputMethodChanged;
 
+	DECLARE_DYNAMIC_MULTICAST_DELEGATE(FOnInputIconUpdated);
+	UPROPERTY(BlueprintAssignable, Category = CommonActionWidget)
+	FOnInputIconUpdated OnInputIconUpdated;
+
 	/**
 	 * The material to use when showing held progress, the progress will be sent using the material parameter
 	 * defined by ProgressMaterialParam and the value will range from 0..1.
