Ensure that brushes are always loaded when needed, even for lazily loaded widgets

REMOVE: UDN for context: https://udn.unrealengine.com/s/question/0D5QP00000N1mXF0AZ/gamepadinputtype-does-not-get-updated-on-controller?fromCase=1

#tests UDN repro fixed PIE

diff --git a/Engine/Plugins/Runtime/CommonUI/Source/CommonInput/Private/CommonInputBaseTypes.cpp b/Engine/Plugins/Runtime/CommonUI/Source/CommonInput/Private/CommonInputBaseTypes.cpp
--- a/Engine/Plugins/Runtime/CommonUI/Source/CommonInput/Private/CommonInputBaseTypes.cpp
+++ b/Engine/Plugins/Runtime/CommonUI/Source/CommonInput/Private/CommonInputBaseTypes.cpp
@@ -309,6 +309,8 @@ bool UCommonInputPlatformSettings::TryGetInputBrush(FSlateBrush& OutBrush, const
 
 FName UCommonInputPlatformSettings::GetBestGamepadNameForHardware(FName CurrentGamepadName, FName InputDeviceName, const FString& HardwareDeviceIdentifier)
 {
+	InitializeControllerData();
+
 	FName FirstMatch = NAME_None;
 
 	// This is far more complicated than it should be because XInput exposes no information about device type,
