PR #12046: Fix Property metadata of ClickMethod from CommonButtonBase

diff --git a/Engine/Plugins/Runtime/CommonUI/Source/CommonUI/Public/CommonButtonBase.h b/Engine/Plugins/Runtime/CommonUI/Source/CommonUI/Public/CommonButtonBase.h
--- a/Engine/Plugins/Runtime/CommonUI/Source/CommonUI/Public/CommonButtonBase.h
+++ b/Engine/Plugins/Runtime/CommonUI/Source/CommonUI/Public/CommonButtonBase.h
@@ -814,7 +814,7 @@ private:
 
 public:
 	/** The type of mouse action required by the user to trigger the button's 'Click' */
-	UPROPERTY(EditAnywhere, BlueprintReadOnly, Category = Input, meta = (ExposeOnSFpawn = true))
+	UPROPERTY(EditAnywhere, BlueprintReadOnly, Category = Input, meta = (ExposeOnSpawn = true))
 	TEnumAsByte<EButtonClickMethod::Type> ClickMethod;
 
 	UPROPERTY(EditAnywhere, BlueprintReadOnly, Category = Input)
