#UE Give all subclasses of FGameplayTag the same struct traits so the functions for those traits actually get called

diff --git a/Engine/Plugins/Runtime/CommonUI/Source/CommonUI/Public/UITag.h b/Engine/Plugins/Runtime/CommonUI/Source/CommonUI/Public/UITag.h
--- a/Engine/Plugins/Runtime/CommonUI/Source/CommonUI/Public/UITag.h
+++ b/Engine/Plugins/Runtime/CommonUI/Source/CommonUI/Public/UITag.h
@@ -106,6 +106,10 @@ struct TStructOpsTypeTraits<TagType> : public TStructOpsTypeTraitsBase2<TagType>
 {	\
 	enum	\
 	{	\
+		WithNetSerializer = true,	\
+		WithNetSharedSerialization = true,	\
+		WithPostSerialize = true,	\
+		WithStructuredSerializeFromMismatchedTag = true,	\
 		WithExportTextItem = true,	\
 		WithImportTextItem = true	\
 	};
