Allow 'SetInputActionBinding' to be called from Blueprint.

This function takes a 'FUIActionBindingHandle' which is a BlueprintType. Exposing this method to Blueprint avoids having to create a native base class for a widget just to pass along a binding handle.

#tests Compiled DevelopmentEditor Win64 / Test Client / Shipping Client / Ran a cooked build and verified the binding passed from blueprint appeared correctly in the action widget.

diff --git a/Engine/Plugins/Runtime/CommonUI/Source/CommonUI/Public/CommonActionWidget.h b/Engine/Plugins/Runtime/CommonUI/Source/CommonUI/Public/CommonActionWidget.h
--- a/Engine/Plugins/Runtime/CommonUI/Source/CommonUI/Public/CommonActionWidget.h
+++ b/Engine/Plugins/Runtime/CommonUI/Source/CommonUI/Public/CommonActionWidget.h
@@ -34,7 +34,7 @@ public:
 	virtual TSharedRef<SWidget> RebuildWidget() override;
 	virtual void ReleaseSlateResources(bool bReleaseChildren) override;
 	virtual void SynchronizeProperties() override;
-	/** End UWidet */
+	/** End UWidget */
 	
 	UFUNCTION(BlueprintCallable, Category = CommonActionWidget)
 	virtual FSlateBrush GetIcon() const;
@@ -51,6 +51,7 @@ public:
 	UFUNCTION(BlueprintCallable, Category = CommonActionWidget)
 	void SetInputAction(FDataTableRowHandle InputActionRow);
 
+	UFUNCTION(BlueprintCallable, Category = CommonActionWidget)
 	void SetInputActionBinding(FUIActionBindingHandle BindingHandle);
 
 	UFUNCTION(BlueprintCallable, Category = CommonActionWidget)
