[Engine]
* Moved dllapi from types to functions/static vars
* Changed so extern template instantiation is based on clang or not instead of if linux and mac since we also build windows using clang
* Added dllapi on some member functions inside template types where the code is in the cpp file. This is needed to tell the compiler these functions will be dll exported (clang does not accept only adding it on the instantiation

diff --git a/Engine/Plugins/Experimental/Mutable/Source/CustomizableObject/Private/MuCO/LogInformationUtil.cpp b/Engine/Plugins/Experimental/Mutable/Source/CustomizableObject/Private/MuCO/LogInformationUtil.cpp
--- a/Engine/Plugins/Experimental/Mutable/Source/CustomizableObject/Private/MuCO/LogInformationUtil.cpp
+++ b/Engine/Plugins/Experimental/Mutable/Source/CustomizableObject/Private/MuCO/LogInformationUtil.cpp
@@ -1,14 +1,13 @@
 // Copyright Epic Games, Inc. All Rights Reserved.
 
 #include "MuCO/LogInformationUtil.h"
-
 #include "MuCO/CustomizableObjectSystemPrivate.h"
 #include "Engine/SkeletalMesh.h"
 #include "Engine/World.h"
 #include "GameFramework/PlayerController.h"
+#include "MuCO/CustomizableObject.h"
 #include "MuCO/CustomizableObjectInstancePrivate.h"
 #include "TextureResource.h"
-#include "MuCO/CustomizableObject.h"
 
 int LogInformationUtil::CountLOD0 = 0;
 int LogInformationUtil::CountLOD1 = 0;
@@ -299,7 +298,7 @@ void LogInformationUtil::PrintImageToPlatformDataMap(const TMap<uint32, FTexture
 		Log = FString::Printf(TEXT("key=%d, "), Elem.Key);
 		Log = FString::Printf(TEXT("value: SizeX=%d, "), Elem.Value->SizeX);
 		Log = FString::Printf(TEXT("SizeY=%d, "), Elem.Value->SizeY);
-		Log = FString::Printf(TEXT("PixelFormat=%d, "), *UEnum::GetValueAsString(Elem.Value->PixelFormat));
+		Log = FString::Printf(TEXT("PixelFormat=%d, "), GetPixelFormatString(Elem.Value->PixelFormat));
 		int i;
 		const int Max = Elem.Value->Mips.Num();
 		for (i = 0; i < Max; ++i)
diff --git a/Engine/Plugins/Online/OnlineSubsystem/Source/Public/OnlineSessionSettings.h b/Engine/Plugins/Online/OnlineSubsystem/Source/Public/OnlineSessionSettings.h
--- a/Engine/Plugins/Online/OnlineSubsystem/Source/Public/OnlineSessionSettings.h
+++ b/Engine/Plugins/Online/OnlineSubsystem/Source/Public/OnlineSessionSettings.h
@@ -185,7 +185,7 @@ public:
 	 * @param ID ID of comparison
 	 */
 	template<typename ValueType>
-	void Set(FName Key, const ValueType& Value, EOnlineComparisonOp::Type InType, int32 ID);
+	ONLINESUBSYSTEM_API void Set(FName Key, const ValueType& Value, EOnlineComparisonOp::Type InType, int32 ID);
 
 	/**
 	 *	Sets a key value pair combination that defines a search parameter
@@ -195,7 +195,7 @@ public:
 	 * @param InType type of comparison
 	 */
 	template<typename ValueType>
-	void Set(FName Key, const ValueType& Value, EOnlineComparisonOp::Type InType);
+	ONLINESUBSYSTEM_API void Set(FName Key, const ValueType& Value, EOnlineComparisonOp::Type InType);
 
 	/**
 	 *	Gets a key value pair combination that defines a search parameter
@@ -206,7 +206,7 @@ public:
 	 * @return true if found, false otherwise
 	 */
 	template<typename ValueType>
-	bool Get(FName Key, ValueType& Value) const;
+	ONLINESUBSYSTEM_API bool Get(FName Key, ValueType& Value) const;
 
 	/**
 	 * Retrieve a search parameter comparison op
@@ -221,7 +221,7 @@ public:
 /**
  *	Container for all settings describing a single online session
  */
-class ONLINESUBSYSTEM_API FOnlineSessionSettings
+class FOnlineSessionSettings
 {
 public:
 	/** The number of publicly available connections advertised */
@@ -306,7 +306,7 @@ public:
 	 * @param ID ID for this session setting
 	 */
 	template<typename ValueType>
-	void Set(FName Key, const ValueType& Value, EOnlineDataAdvertisementType::Type InType, int32 InID);
+	ONLINESUBSYSTEM_API void Set(FName Key, const ValueType& Value, EOnlineDataAdvertisementType::Type InType, int32 InID);
 
 	/**
 	 *	Sets a key value pair combination that defines a session setting
@@ -316,7 +316,7 @@ public:
 	 * @param InType type of online advertisement
 	 */
 	template<typename ValueType>
-	void Set(FName Key, const ValueType& Value, EOnlineDataAdvertisementType::Type InType);
+	ONLINESUBSYSTEM_API void Set(FName Key, const ValueType& Value, EOnlineDataAdvertisementType::Type InType);
 
 	/**
 	 *	Sets a key value pair combination that defines a session setting
@@ -325,7 +325,7 @@ public:
 	 * @param Key key for the setting
 	 * @param SrcSetting setting values
 	 */
-	void Set(FName Key, const FOnlineSessionSetting& SrcSetting);
+	ONLINESUBSYSTEM_API void Set(FName Key, const FOnlineSessionSetting& SrcSetting);
 
 	/**
 	 *	Gets a key value pair combination that defines a session setting
@@ -336,7 +336,7 @@ public:
 	 * @return true if found, false otherwise
 	 */
 	template<typename ValueType>
-	bool Get(FName Key, ValueType& Value) const;
+	ONLINESUBSYSTEM_API bool Get(FName Key, ValueType& Value) const;
 
 	/**
 	 *  Removes a key value pair combination
@@ -345,7 +345,7 @@ public:
 	 *
 	 * @return true if found and removed, false otherwise
 	 */
-	bool Remove(FName Key);
+	ONLINESUBSYSTEM_API bool Remove(FName Key);
 
 	/**
 	 * Retrieve a session setting's advertisement type
@@ -354,7 +354,7 @@ public:
 	 *
 	 * @return the advertisement type for the setting
 	 */
-	EOnlineDataAdvertisementType::Type GetAdvertisementType(FName Key) const;
+	ONLINESUBSYSTEM_API EOnlineDataAdvertisementType::Type GetAdvertisementType(FName Key) const;
 
 	/**
 	* Retrieve a session setting's ID
@@ -363,7 +363,7 @@ public:
 	*
 	* @return the ID for the setting
 	*/
-	int32 GetID(FName Key) const;
+	ONLINESUBSYSTEM_API int32 GetID(FName Key) const;
 
 };
 
diff --git a/Engine/Plugins/Runtime/AudioSynesthesia/Source/AudioSynesthesia/Classes/AudioSynesthesia.h b/Engine/Plugins/Runtime/AudioSynesthesia/Source/AudioSynesthesia/Classes/AudioSynesthesia.h
--- a/Engine/Plugins/Runtime/AudioSynesthesia/Source/AudioSynesthesia/Classes/AudioSynesthesia.h
+++ b/Engine/Plugins/Runtime/AudioSynesthesia/Source/AudioSynesthesia/Classes/AudioSynesthesia.h
@@ -8,15 +8,15 @@
  *
  * Defines asset actions for derived UAudioSynthesiaSettings subclasses.
  */
-UCLASS(Abstract, Blueprintable)
-class AUDIOSYNESTHESIA_API UAudioSynesthesiaSettings : public UAudioAnalyzerSettings
+UCLASS(Abstract, Blueprintable, MinimalAPI)
+class UAudioSynesthesiaSettings : public UAudioAnalyzerSettings
 {
 	GENERATED_BODY()
 
 public:
 
 #if WITH_EDITOR
-	const TArray<FText>& GetAssetActionSubmenus() const;
-	FColor GetTypeColor() const override;
+	AUDIOSYNESTHESIA_API const TArray<FText>& GetAssetActionSubmenus() const;
+	AUDIOSYNESTHESIA_API FColor GetTypeColor() const override;
 #endif
 };
\ No newline at end of file
diff --git a/Engine/Plugins/Runtime/AudioSynesthesia/Source/AudioSynesthesia/Classes/AudioSynesthesiaNRT.h b/Engine/Plugins/Runtime/AudioSynesthesia/Source/AudioSynesthesia/Classes/AudioSynesthesiaNRT.h
--- a/Engine/Plugins/Runtime/AudioSynesthesia/Source/AudioSynesthesia/Classes/AudioSynesthesiaNRT.h
+++ b/Engine/Plugins/Runtime/AudioSynesthesia/Source/AudioSynesthesia/Classes/AudioSynesthesiaNRT.h
@@ -8,17 +8,17 @@
  *
  * Defines asset actions for derived UAudioSynthesiaNRTSettings subclasses.
  */
-UCLASS(Abstract, Blueprintable)
-class AUDIOSYNESTHESIA_API UAudioSynesthesiaNRTSettings : public UAudioAnalyzerNRTSettings
+UCLASS(Abstract, Blueprintable, MinimalApi)
+class UAudioSynesthesiaNRTSettings : public UAudioAnalyzerNRTSettings
 {
 	GENERATED_BODY()
 
 	public:
 
-		const TArray<FText>& GetAssetActionSubmenus() const;
+		AUDIOSYNESTHESIA_API const TArray<FText>& GetAssetActionSubmenus() const;
 
 #if WITH_EDITOR
-		FColor GetTypeColor() const override;
+		AUDIOSYNESTHESIA_API FColor GetTypeColor() const override;
 #endif
 };
 
@@ -26,17 +26,17 @@ class AUDIOSYNESTHESIA_API UAudioSynesthesiaNRTSettings : public UAudioAnalyzerN
  *
  * Defines asset actions for derived UAudioSynthesiaNRT subclasses.
  */
-UCLASS(Abstract, Blueprintable)
-class AUDIOSYNESTHESIA_API UAudioSynesthesiaNRT : public UAudioAnalyzerNRT
+UCLASS(Abstract, Blueprintable, MinimalApi)
+class UAudioSynesthesiaNRT : public UAudioAnalyzerNRT
 {
 	GENERATED_BODY()
 
 	public:
 
-		const TArray<FText>& GetAssetActionSubmenus() const;
+		AUDIOSYNESTHESIA_API const TArray<FText>& GetAssetActionSubmenus() const;
 
 #if WITH_EDITOR
-		FColor GetTypeColor() const override;
+		AUDIOSYNESTHESIA_API FColor GetTypeColor() const override;
 #endif
 };
 
diff --git a/Engine/Plugins/Runtime/CommonUI/Source/CommonUI/Public/CommonButtonBase.h b/Engine/Plugins/Runtime/CommonUI/Source/CommonUI/Public/CommonButtonBase.h
--- a/Engine/Plugins/Runtime/CommonUI/Source/CommonUI/Public/CommonButtonBase.h
+++ b/Engine/Plugins/Runtime/CommonUI/Source/CommonUI/Public/CommonButtonBase.h
@@ -55,12 +55,12 @@ public:
 /* ---- All properties must be EditDefaultsOnly, BlueprintReadOnly !!! -----
  *       we return the CDO to blueprints, so we cannot allow any changes (blueprint doesn't support const variables)
  */
-UCLASS(Abstract, Blueprintable, ClassGroup = UI, meta = (Category = "Common UI"))
-class COMMONUI_API UCommonButtonStyle : public UObject
+UCLASS(Abstract, Blueprintable, MinimalAPI, ClassGroup = UI, meta = (Category = "Common UI"))
+class UCommonButtonStyle : public UObject
 {
 	GENERATED_BODY()
 
-	virtual bool NeedsLoadForServer() const override;
+	COMMONUI_API virtual bool NeedsLoadForServer() const override;
 
 public:
 	
@@ -161,71 +161,70 @@ public:
 	FCommonButtonStyleOptionalSlateSound LockedHoveredSlateSound;
 
 	UFUNCTION(BlueprintCallable, Category = "Common ButtonStyle|Getters")
-	void GetMaterialBrush(FSlateBrush& Brush) const;
+	COMMONUI_API void GetMaterialBrush(FSlateBrush& Brush) const;
 
 	UFUNCTION(BlueprintCallable, Category = "Common ButtonStyle|Getters")
-	void GetNormalBaseBrush(FSlateBrush& Brush) const;
+	COMMONUI_API void GetNormalBaseBrush(FSlateBrush& Brush) const;
 
 	UFUNCTION(BlueprintCallable, Category = "Common ButtonStyle|Getters")
-	void GetNormalHoveredBrush(FSlateBrush& Brush) const;
+	COMMONUI_API void GetNormalHoveredBrush(FSlateBrush& Brush) const;
 
 	UFUNCTION(BlueprintCallable, Category = "Common ButtonStyle|Getters")
-	void GetNormalPressedBrush(FSlateBrush& Brush) const;
+	COMMONUI_API void GetNormalPressedBrush(FSlateBrush& Brush) const;
 
 	UFUNCTION(BlueprintCallable, Category = "Common ButtonStyle|Getters")
-	void GetSelectedBaseBrush(FSlateBrush& Brush) const;
+	COMMONUI_API void GetSelectedBaseBrush(FSlateBrush& Brush) const;
 
 	UFUNCTION(BlueprintCallable, Category = "Common ButtonStyle|Getters")
-	void GetSelectedHoveredBrush(FSlateBrush& Brush) const;
+	COMMONUI_API void GetSelectedHoveredBrush(FSlateBrush& Brush) const;
 
 	UFUNCTION(BlueprintCallable, Category = "Common ButtonStyle|Getters")
-	void GetSelectedPressedBrush(FSlateBrush& Brush) const;
+	COMMONUI_API void GetSelectedPressedBrush(FSlateBrush& Brush) const;
 
 	UFUNCTION(BlueprintCallable, Category = "Common ButtonStyle|Getters")
-	void GetDisabledBrush(FSlateBrush& Brush) const;
+	COMMONUI_API void GetDisabledBrush(FSlateBrush& Brush) const;
 
 	UFUNCTION(BlueprintCallable, Category = "Common ButtonStyle|Getters")
-	void GetButtonPadding(FMargin& OutButtonPadding) const;
+	COMMONUI_API void GetButtonPadding(FMargin& OutButtonPadding) const;
 
 	UFUNCTION(BlueprintCallable, Category = "Common ButtonStyle|Getters")
-	void GetCustomPadding(FMargin& OutCustomPadding) const;
+	COMMONUI_API void GetCustomPadding(FMargin& OutCustomPadding) const;
 
 	UFUNCTION(BlueprintCallable, Category = "Common ButtonStyle|Getters")
-	UCommonTextStyle* GetNormalTextStyle() const;
+	COMMONUI_API UCommonTextStyle* GetNormalTextStyle() const;
 
 	UFUNCTION(BlueprintCallable, Category = "Common ButtonStyle|Getters")
-	UCommonTextStyle* GetNormalHoveredTextStyle() const;
+	COMMONUI_API UCommonTextStyle* GetNormalHoveredTextStyle() const;
 
 	UFUNCTION(BlueprintCallable, Category = "Common ButtonStyle|Getters")
-	UCommonTextStyle* GetSelectedTextStyle() const;
+	COMMONUI_API UCommonTextStyle* GetSelectedTextStyle() const;
 
 	UFUNCTION(BlueprintCallable, Category = "Common ButtonStyle|Getters")
-	UCommonTextStyle* GetSelectedHoveredTextStyle() const;
+	COMMONUI_API UCommonTextStyle* GetSelectedHoveredTextStyle() const;
 
 	UFUNCTION(BlueprintCallable, Category = "Common ButtonStyle|Getters")
-	UCommonTextStyle* GetDisabledTextStyle() const;
-
+	COMMONUI_API UCommonTextStyle* GetDisabledTextStyle() const;
 };
 
 DECLARE_DELEGATE_RetVal(FReply, FOnButtonDoubleClickedEvent);
 
 /** Custom UButton override that allows us to disable clicking without disabling the widget entirely */
-UCLASS(Experimental)	// "Experimental" to hide it in the designer
-class COMMONUI_API UCommonButtonInternalBase : public UButton
+UCLASS(Experimental, MinimalAPI)	// "Experimental" to hide it in the designer
+class UCommonButtonInternalBase : public UButton
 {
 	GENERATED_UCLASS_BODY()
 
 public:
-	void SetButtonEnabled(bool bInIsButtonEnabled);
-	void SetInteractionEnabled(bool bInIsInteractionEnabled);
+	COMMONUI_API void SetButtonEnabled(bool bInIsButtonEnabled);
+	COMMONUI_API void SetInteractionEnabled(bool bInIsInteractionEnabled);
 
 	/** Updates the IsFocusable flag and updates the bIsFocusable flag of the underlying slate button widget */
-	void SetButtonFocusable(bool bInIsButtonFocusable);
-	bool IsHovered() const;
-	bool IsPressed() const;
+	COMMONUI_API void SetButtonFocusable(bool bInIsButtonFocusable);
+	COMMONUI_API bool IsHovered() const;
+	COMMONUI_API bool IsPressed() const;
 
-	void SetMinDesiredHeight(int32 InMinHeight);
-	void SetMinDesiredWidth(int32 InMinWidth);
+	COMMONUI_API void SetMinDesiredHeight(int32 InMinHeight);
+	COMMONUI_API void SetMinDesiredWidth(int32 InMinWidth);
 
 	/** Called when the button is clicked */
 	FOnButtonDoubleClickedEvent HandleDoubleClicked;
@@ -242,20 +241,20 @@ public:
 
 protected:
 	// UWidget interface
-	virtual TSharedRef<SWidget> RebuildWidget() override;
-	virtual void ReleaseSlateResources(bool bReleaseChildren) override;
+	COMMONUI_API virtual TSharedRef<SWidget> RebuildWidget() override;
+	COMMONUI_API virtual void ReleaseSlateResources(bool bReleaseChildren) override;
 	// End of UWidget interface
 
-	virtual FReply SlateHandleClickedOverride();
-	virtual void SlateHandlePressedOverride();
-	virtual void SlateHandleReleasedOverride();
-	virtual FReply SlateHandleDoubleClicked();
+	COMMONUI_API virtual FReply SlateHandleClickedOverride();
+	COMMONUI_API virtual void SlateHandlePressedOverride();
+	COMMONUI_API virtual void SlateHandleReleasedOverride();
+	COMMONUI_API virtual FReply SlateHandleDoubleClicked();
 
 	/** Called when internal slate button receives focus; Fires OnReceivedFocus */
-	void SlateHandleOnReceivedFocus();
+	COMMONUI_API void SlateHandleOnReceivedFocus();
 
 	/** Called when internal slate button loses focus; Fires OnLostFocus */
-	void SlateHandleOnLostFocus();
+	COMMONUI_API void SlateHandleOnLostFocus();
 
 	/** The minimum width of the button */
 	UPROPERTY()
@@ -287,150 +286,150 @@ DECLARE_DYNAMIC_MULTICAST_DELEGATE_OneParam(FCommonButtonBaseClicked, class UCom
 /**
  * Button that disables itself when not active. Also updates actions for CommonActionWidget if bound to display platform-specific icons.
  */
-UCLASS(Abstract, Blueprintable, ClassGroup = UI, meta = (Category = "Common UI", DisableNativeTick))
-class COMMONUI_API UCommonButtonBase : public UCommonUserWidget
+UCLASS(Abstract, Blueprintable, MinimalAPI, ClassGroup = UI, meta = (Category = "Common UI", DisableNativeTick))
+class UCommonButtonBase : public UCommonUserWidget
 {
 	GENERATED_UCLASS_BODY()
 
 public:
 	// UWidget interface
-	virtual bool IsHovered() const override;
+	COMMONUI_API virtual bool IsHovered() const override;
 	// End of UWidget interface
 
 	// UUserWidget interface
-	virtual void NativeConstruct() override;
-	virtual void NativeDestruct() override;
-	virtual bool Initialize() override;
-	virtual void SetIsEnabled(bool bInIsEnabled) override;
-	virtual bool NativeIsInteractable() const override;
+	COMMONUI_API virtual void NativeConstruct() override;
+	COMMONUI_API virtual void NativeDestruct() override;
+	COMMONUI_API virtual bool Initialize() override;
+	COMMONUI_API virtual void SetIsEnabled(bool bInIsEnabled) override;
+	COMMONUI_API virtual bool NativeIsInteractable() const override;
 	// End of UUserWidget interface
 	
 	/** Disables this button with a reason (use instead of SetIsEnabled) */
 	UFUNCTION(BlueprintCallable, Category = "Common Button|Setters")
-	void DisableButtonWithReason(const FText& DisabledReason);
+	COMMONUI_API void DisableButtonWithReason(const FText& DisabledReason);
 
 	/** Change whether this widget is selectable at all. If false and currently selected, will deselect. */
 	UFUNCTION(BlueprintCallable, Category = "Common Button|Setters")
-	void SetIsInteractionEnabled(bool bInIsInteractionEnabled);
+	COMMONUI_API void SetIsInteractionEnabled(bool bInIsInteractionEnabled);
 
 	UFUNCTION(BlueprintCallable, Category = "Common Button|Setters")
-	void SetHideInputAction(bool bInHideInputAction);
+	COMMONUI_API void SetHideInputAction(bool bInHideInputAction);
 
 	/** Is this button currently interactable? (use instead of GetIsEnabled) */
 	UFUNCTION(BlueprintCallable, Category = "Common Button|Getters")
-	bool IsInteractionEnabled() const;
+	COMMONUI_API bool IsInteractionEnabled() const;
 
 	/** Is this button currently pressed? */
 	UFUNCTION(BlueprintCallable, Category = "Common Button|Getters")
-	bool IsPressed() const;
+	COMMONUI_API bool IsPressed() const;
 
 	/** Set the click method for mouse interaction */
 	UFUNCTION(BlueprintCallable, Category = "Common Button|Setters")
-	void SetClickMethod(EButtonClickMethod::Type InClickMethod);
+	COMMONUI_API void SetClickMethod(EButtonClickMethod::Type InClickMethod);
 
 	/** Set the click method for touch interaction */
 	UFUNCTION(BlueprintCallable, Category = "Common Button|Setters")
-	void SetTouchMethod(EButtonTouchMethod::Type InTouchMethod);
+	COMMONUI_API void SetTouchMethod(EButtonTouchMethod::Type InTouchMethod);
 
 	/** Set the click method for keyboard/gamepad button press interaction */
 	UFUNCTION(BlueprintCallable, Category = "Common Button|Setters")
-	void SetPressMethod(EButtonPressMethod::Type InPressMethod);
+	COMMONUI_API void SetPressMethod(EButtonPressMethod::Type InPressMethod);
 
 	/** Change whether this widget is selectable at all. If false and currently selected, will deselect. */
 	UFUNCTION(BlueprintCallable, Category = "Common Button|Setters")
-	void SetIsSelectable(bool bInIsSelectable);
+	COMMONUI_API void SetIsSelectable(bool bInIsSelectable);
 
 	/** Change whether this widget is selectable at all. If false and currently selected, will deselect. */
 	UFUNCTION(BlueprintCallable, Category = "Common Button|Setters")
-	void SetIsInteractableWhenSelected(bool bInInteractableWhenSelected);
+	COMMONUI_API void SetIsInteractableWhenSelected(bool bInInteractableWhenSelected);
 
 	/** Change whether this widget is toggleable. If toggleable, clicking when selected will deselect. */
 	UFUNCTION(BlueprintCallable, Category = "Common Button|Setters")
-	void SetIsToggleable(bool bInIsToggleable);
+	COMMONUI_API void SetIsToggleable(bool bInIsToggleable);
 
 	/** Change whether this widget should use the fallback default input action. */
 	UFUNCTION(BlueprintCallable, Category = "Common Button|Setters")
-	void SetShouldUseFallbackDefaultInputAction(bool bInShouldUseFallbackDefaultInputAction);
+	COMMONUI_API void SetShouldUseFallbackDefaultInputAction(bool bInShouldUseFallbackDefaultInputAction);
 
 	/** 
 	 * Change the selected state manually.
 	 * @param bGiveClickFeedback	If true, the button may give user feedback as if it were clicked. IE: Play a click sound, trigger animations as if it were clicked.
 	 */
 	UFUNCTION(BlueprintCallable, Category = "Common Button|Setters")
-	void SetIsSelected(bool InSelected, bool bGiveClickFeedback = true);
+	COMMONUI_API void SetIsSelected(bool InSelected, bool bGiveClickFeedback = true);
 
 	/** Change whether this widget is locked. If locked, the button can be focusable and responsive to mouse input but will not broadcast OnClicked events. */
 	UFUNCTION(BlueprintCallable, Category = "Common Button|Setters")
-	void SetIsLocked(bool bInIsLocked);
+	COMMONUI_API void SetIsLocked(bool bInIsLocked);
 
 	/** @returns True if the button is currently in a selected state, False otherwise */
 	UFUNCTION(BlueprintCallable, Category = "Common Button|Getters")
-	bool GetSelected() const;
+	COMMONUI_API bool GetSelected() const;
 
 	/** @returns True if the button is currently locked, False otherwise */
 	UFUNCTION(BlueprintCallable, Category = "Common Button|Getters")
-	bool GetLocked() const;
+	COMMONUI_API bool GetLocked() const;
 
 	UFUNCTION(BlueprintCallable, Category = "Common Button" )
-	void ClearSelection();
+	COMMONUI_API void ClearSelection();
 
 	/** Set whether the button should become selected upon receiving focus or not; Only settable for buttons that are selectable */
 	UFUNCTION(BlueprintCallable, Category = "Common Button|Setters")
-	void SetShouldSelectUponReceivingFocus(bool bInShouldSelectUponReceivingFocus);
+	COMMONUI_API void SetShouldSelectUponReceivingFocus(bool bInShouldSelectUponReceivingFocus);
 
 	/** Get whether the button should become selected upon receiving focus or not */
 	UFUNCTION(BlueprintCallable, Category = "Common Button|Getters")
-	bool GetShouldSelectUponReceivingFocus() const;
+	COMMONUI_API bool GetShouldSelectUponReceivingFocus() const;
 
 	/** Sets the style of this button, rebuilds the internal styling */
 	UFUNCTION(BlueprintCallable, Category = "Common Button|Setters")
-	void SetStyle(TSubclassOf<UCommonButtonStyle> InStyle = nullptr);
+	COMMONUI_API void SetStyle(TSubclassOf<UCommonButtonStyle> InStyle = nullptr);
 
 	/** @Returns Current button style*/
 	UFUNCTION(BlueprintCallable, Category = "Common Button|Getters")
-	UCommonButtonStyle* GetStyle() const;
+	COMMONUI_API UCommonButtonStyle* GetStyle() const;
 
-	const UCommonButtonStyle* GetStyleCDO() const;
+	COMMONUI_API const UCommonButtonStyle* GetStyleCDO() const;
 
 	/** @return The current button padding that corresponds to the current size and selection state */
 	UFUNCTION(BlueprintCallable, Category = "Common Button|Getters")
-	void GetCurrentButtonPadding(FMargin& OutButtonPadding) const;
+	COMMONUI_API void GetCurrentButtonPadding(FMargin& OutButtonPadding) const;
 
 	/** @return The custom padding that corresponds to the current size and selection state */
 	UFUNCTION(BlueprintCallable, Category = "Common Button|Getters")
-	void GetCurrentCustomPadding(FMargin& OutCustomPadding) const;
+	COMMONUI_API void GetCurrentCustomPadding(FMargin& OutCustomPadding) const;
 	
 	/** @return The text style that corresponds to the current size and selection state */
 	UFUNCTION(BlueprintCallable, Category = "Common Button|Getters")
-	UCommonTextStyle* GetCurrentTextStyle() const;
+	COMMONUI_API UCommonTextStyle* GetCurrentTextStyle() const;
 
 	/** @return The class of the text style that corresponds to the current size and selection state */
 	UFUNCTION(BlueprintCallable, Category = "Common Button|Getters")
-	TSubclassOf<UCommonTextStyle> GetCurrentTextStyleClass() const;
+	COMMONUI_API TSubclassOf<UCommonTextStyle> GetCurrentTextStyleClass() const;
 
 	/** Sets the minimum dimensions of this button */
 	UFUNCTION(BlueprintCallable, Category = "Common Button|Setters")
-	void SetMinDimensions(int32 InMinWidth, int32 InMinHeight);
+	COMMONUI_API void SetMinDimensions(int32 InMinWidth, int32 InMinHeight);
 
 	/** Updates the current triggered action */
 	UFUNCTION(BlueprintCallable, Category = "Common Button|Setters")
-	void SetTriggeredInputAction(const FDataTableRowHandle &InputActionRow);
+	COMMONUI_API void SetTriggeredInputAction(const FDataTableRowHandle &InputActionRow);
 
 	/** Updates the current triggering action */
 	UFUNCTION(BlueprintCallable, Category = "Common Button|Setters")
-	void SetTriggeringInputAction(const FDataTableRowHandle & InputActionRow);
+	COMMONUI_API void SetTriggeringInputAction(const FDataTableRowHandle & InputActionRow);
 
 	/** Updates the current triggering enhanced input action, requires enhanced input enabled in CommonUI settings */
 	UFUNCTION(BlueprintCallable, Category = "Common Button|Setters")
-	void SetTriggeringEnhancedInputAction(UInputAction* InInputAction);
+	COMMONUI_API void SetTriggeringEnhancedInputAction(UInputAction* InInputAction);
 
 	/** Gets the appropriate input action that is set */
 	UFUNCTION(BlueprintCallable, Category = "Common Button|Getters")
-	bool GetInputAction(FDataTableRowHandle &InputActionRow) const;
+	COMMONUI_API bool GetInputAction(FDataTableRowHandle &InputActionRow) const;
 
 	/** Gets the appropriate enhanced input action that is set */
 	UFUNCTION(BlueprintCallable, Category = "Common Button|Getters")
-	UInputAction* GetEnhancedInputAction() const;
+	COMMONUI_API UInputAction* GetEnhancedInputAction() const;
 
 	/** Returns true if this button has a hold behavior, even if the triggering action is not holdable. */
 	UFUNCTION(BlueprintCallable, BlueprintPure, Category = "Common Button|Getters")
@@ -438,7 +437,7 @@ public:
 
 	/** Change whether this button should have a hold behavior even if the triggering action is not holdable. */
 	UFUNCTION(BlueprintCallable, Category = "Common Button|Setters")
-	void SetRequiresHold(bool bInRequiresHold);
+	COMMONUI_API void SetRequiresHold(bool bInRequiresHold);
 
 	/** Returns required hold time for performing a triggering action. */
 	UFUNCTION(BlueprintCallable, BlueprintPure, Category = "Common Button|Getters")
@@ -446,36 +445,36 @@ public:
 
 	/** Updates the bIsFocusable flag */
 	UFUNCTION(BlueprintCallable, Category = "Common Button|Getters")
-	void SetIsFocusable(bool bInIsFocusable);
+	COMMONUI_API void SetIsFocusable(bool bInIsFocusable);
 
 	/** Gets the bIsFocusable flag */
 	UFUNCTION(BlueprintCallable, Category = "Common Button|Getters")
-	bool GetIsFocusable() const;
+	COMMONUI_API bool GetIsFocusable() const;
 
 	/** Returns the dynamic instance of the material being used for this button, if it is using a single material style. */
 	UFUNCTION(BlueprintCallable, BlueprintPure, Category = "Common Button|Getters")
-	UMaterialInstanceDynamic* GetSingleMaterialStyleMID() const;
+	COMMONUI_API UMaterialInstanceDynamic* GetSingleMaterialStyleMID() const;
 
 	UFUNCTION(BlueprintCallable, Category = "Common Button|Input")
-	void SetInputActionProgressMaterial(const FSlateBrush& InProgressMaterialBrush, const FName& InProgressMaterialParam);
+	COMMONUI_API void SetInputActionProgressMaterial(const FSlateBrush& InProgressMaterialBrush, const FName& InProgressMaterialParam);
 
 	UFUNCTION(BlueprintCallable, Category = "Common Button|Sound")
-	void SetPressedSoundOverride(USoundBase* Sound);
+	COMMONUI_API void SetPressedSoundOverride(USoundBase* Sound);
 
 	UFUNCTION(BlueprintCallable, Category = "Common Button|Sound")
-	void SetHoveredSoundOverride(USoundBase* Sound);
+	COMMONUI_API void SetHoveredSoundOverride(USoundBase* Sound);
 
 	UFUNCTION(BlueprintCallable, Category = "Common Button|Sound")
-	void SetSelectedPressedSoundOverride(USoundBase* Sound);
+	COMMONUI_API void SetSelectedPressedSoundOverride(USoundBase* Sound);
 
 	UFUNCTION(BlueprintCallable, Category = "Common Button|Sound")
-	void SetSelectedHoveredSoundOverride(USoundBase* Sound);
+	COMMONUI_API void SetSelectedHoveredSoundOverride(USoundBase* Sound);
 	
 	UFUNCTION(BlueprintCallable, Category = "Common Button|Sound")
-	void SetLockedPressedSoundOverride(USoundBase* Sound);
+	COMMONUI_API void SetLockedPressedSoundOverride(USoundBase* Sound);
 
 	UFUNCTION(BlueprintCallable, Category = "Common Button|Sound")
-	void SetLockedHoveredSoundOverride(USoundBase* Sound);
+	COMMONUI_API void SetLockedHoveredSoundOverride(USoundBase* Sound);
 
 	DECLARE_EVENT(UCommonButtonBase, FCommonButtonEvent);
 	FCommonButtonEvent& OnClicked() const { return OnClickedEvent; }
@@ -499,206 +498,206 @@ protected:
 	bool IsPersistentBinding() const { return bIsPersistentBinding; }
 	ECommonInputMode GetInputModeOverride() const { return InputModeOverride; }
 	
-	virtual UCommonButtonInternalBase* ConstructInternalButton();
+	COMMONUI_API virtual UCommonButtonInternalBase* ConstructInternalButton();
 
-	virtual void OnWidgetRebuilt();
-	virtual void PostLoad() override;
-	virtual void SynchronizeProperties() override;
-	virtual FReply NativeOnFocusReceived(const FGeometry& InGeometry, const FFocusEvent& InFocusEvent) override;
-	virtual void NativeOnFocusLost(const FFocusEvent& InFocusEvent) override;
+	COMMONUI_API virtual void OnWidgetRebuilt();
+	COMMONUI_API virtual void PostLoad() override;
+	COMMONUI_API virtual void SynchronizeProperties() override;
+	COMMONUI_API virtual FReply NativeOnFocusReceived(const FGeometry& InGeometry, const FFocusEvent& InFocusEvent) override;
+	COMMONUI_API virtual void NativeOnFocusLost(const FFocusEvent& InFocusEvent) override;
 
 #if WITH_EDITOR
-	virtual void OnCreationFromPalette() override;
-	const FText GetPaletteCategory() override;
+	COMMONUI_API virtual void OnCreationFromPalette() override;
+	COMMONUI_API const FText GetPaletteCategory() override;
 #endif // WITH_EDITOR
 
 	// In some scenarios, USoundBase* setters are bypassed, so standard setters must also be provided for sound override properties
-	void SetPressedSlateSoundOverride(const FSlateSound& InPressedSlateSoundOverride);
-	void SetHoveredSlateSoundOverride(const FSlateSound& InHoveredSlateSoundOverride);
-	void SetSelectedPressedSlateSoundOverride(const FSlateSound& InSelectedPressedSlateSoundOverride);
-	void SetSelectedHoveredSlateSoundOverride(const FSlateSound& InSelectedHoveredSlateSoundOverride);
-	void SetLockedPressedSlateSoundOverride(const FSlateSound& InLockedPressedSlateSoundOverride);
-	void SetLockedHoveredSlateSoundOverride(const FSlateSound& InLockedHoveredSlateSoundOverride);
+	COMMONUI_API void SetPressedSlateSoundOverride(const FSlateSound& InPressedSlateSoundOverride);
+	COMMONUI_API void SetHoveredSlateSoundOverride(const FSlateSound& InHoveredSlateSoundOverride);
+	COMMONUI_API void SetSelectedPressedSlateSoundOverride(const FSlateSound& InSelectedPressedSlateSoundOverride);
+	COMMONUI_API void SetSelectedHoveredSlateSoundOverride(const FSlateSound& InSelectedHoveredSlateSoundOverride);
+	COMMONUI_API void SetLockedPressedSlateSoundOverride(const FSlateSound& InLockedPressedSlateSoundOverride);
+	COMMONUI_API void SetLockedHoveredSlateSoundOverride(const FSlateSound& InLockedHoveredSlateSoundOverride);
 
 	/** Helper function to bind to input method change events */
-	virtual void BindInputMethodChangedDelegate();
+	COMMONUI_API virtual void BindInputMethodChangedDelegate();
 
 	/** Helper function to unbind from input method change events */
-	virtual void UnbindInputMethodChangedDelegate();
+	COMMONUI_API virtual void UnbindInputMethodChangedDelegate();
 
 	/** Called via delegate when the input method changes */
 	UFUNCTION()
-	virtual void OnInputMethodChanged(ECommonInputType CurrentInputType);
+	COMMONUI_API virtual void OnInputMethodChanged(ECommonInputType CurrentInputType);
 
 	/** If HoldData is valid, assigns its values to Keyboard and Mouse, Gamepad and Touch, based off the Current Input Type. */
     UFUNCTION()
-	virtual void UpdateHoldData(ECommonInputType CurrentInputType);
+	COMMONUI_API virtual void UpdateHoldData(ECommonInputType CurrentInputType);
 
 	/** Associates this button at its priority with the given key */
-	virtual void BindTriggeringInputActionToClick();
+	COMMONUI_API virtual void BindTriggeringInputActionToClick();
 
 	/** Associates this button at its priority with the given key */
-	virtual void UnbindTriggeringInputActionToClick();
+	COMMONUI_API virtual void UnbindTriggeringInputActionToClick();
 
 	UFUNCTION()
-	virtual void HandleTriggeringActionCommited(bool& bPassthrough);
-	virtual void HandleTriggeringActionCommited();
+	COMMONUI_API virtual void HandleTriggeringActionCommited(bool& bPassthrough);
+	COMMONUI_API virtual void HandleTriggeringActionCommited();
 
 	// @TODO: DarenC - API decision, consider removing
-	virtual void ExecuteTriggeredInput();
+	COMMONUI_API virtual void ExecuteTriggeredInput();
 
 	/** Helper function to update the associated input action widget, if any, based upon the state of the button */
-	virtual void UpdateInputActionWidget();
+	COMMONUI_API virtual void UpdateInputActionWidget();
 
 	/** Handler function registered to the underlying button's click. */
 	UFUNCTION()
-	void HandleButtonClicked();
+	COMMONUI_API void HandleButtonClicked();
 
 	/** Handler function registered to the underlying button's double click. */
-	FReply HandleButtonDoubleClicked();
+	COMMONUI_API FReply HandleButtonDoubleClicked();
 
 	/** Helper function registered to the underlying button receiving focus */
 	UFUNCTION()
-	virtual void HandleFocusReceived();
+	COMMONUI_API virtual void HandleFocusReceived();
 	
 	/** Helper function registered to the underlying button losing focus */
 	UFUNCTION()
-	virtual void HandleFocusLost();
+	COMMONUI_API virtual void HandleFocusLost();
 
 	/** Helper function registered to the underlying button when pressed */
 	UFUNCTION()
-	void HandleButtonPressed();
+	COMMONUI_API void HandleButtonPressed();
 
 	/** Helper function registered to the underlying button when released */
 	UFUNCTION()
-	void HandleButtonReleased();
+	COMMONUI_API void HandleButtonReleased();
 
 	UFUNCTION(BlueprintImplementableEvent, Category = CommonButton, meta = (DisplayName = "On Selected"))
-	void BP_OnSelected();
-	virtual void NativeOnSelected(bool bBroadcast);
+	COMMONUI_API void BP_OnSelected();
+	COMMONUI_API virtual void NativeOnSelected(bool bBroadcast);
 
 	UFUNCTION(BlueprintImplementableEvent, Category = CommonButton, meta = (DisplayName = "On Deselected"))
-	void BP_OnDeselected();
-	virtual void NativeOnDeselected(bool bBroadcast);
+	COMMONUI_API void BP_OnDeselected();
+	COMMONUI_API virtual void NativeOnDeselected(bool bBroadcast);
 
 	UFUNCTION(BlueprintImplementableEvent, Category = CommonButton, meta = (DisplayName = "On Hovered"))
-	void BP_OnHovered();
-	virtual void NativeOnHovered();
+	COMMONUI_API void BP_OnHovered();
+	COMMONUI_API virtual void NativeOnHovered();
 
 	UFUNCTION(BlueprintImplementableEvent, Category = CommonButton, meta = (DisplayName = "On Unhovered"))
-	void BP_OnUnhovered();
-	virtual void NativeOnUnhovered();
+	COMMONUI_API void BP_OnUnhovered();
+	COMMONUI_API virtual void NativeOnUnhovered();
 	
 	UFUNCTION(BlueprintImplementableEvent, Category = CommonButton, meta = (DisplayName = "On Focused"))
-	void BP_OnFocusReceived();
+	COMMONUI_API void BP_OnFocusReceived();
 
 	UFUNCTION(BlueprintImplementableEvent, Category = CommonButton, meta = (DisplayName = "On Unfocused"))
-	void BP_OnFocusLost();
+	COMMONUI_API void BP_OnFocusLost();
 
 	UFUNCTION(BlueprintImplementableEvent, Category = CommonButton, meta = (DisplayName = "On Locked Changed"))
-	void BP_OnLockedChanged(bool bIsLocked);
+	COMMONUI_API void BP_OnLockedChanged(bool bIsLocked);
 	
 	UFUNCTION(BlueprintImplementableEvent, Category = CommonButton, meta = (DisplayName = "On Lock Clicked"))
-	void BP_OnLockClicked();
+	COMMONUI_API void BP_OnLockClicked();
 
 	UFUNCTION(BlueprintImplementableEvent, Category = CommonButton, meta = (DisplayName = "On Lock Double Clicked"))
-	void BP_OnLockDoubleClicked();
+	COMMONUI_API void BP_OnLockDoubleClicked();
 
 	UFUNCTION(BlueprintImplementableEvent, Category = CommonButton, meta = (DisplayName = "On Clicked"))
-	void BP_OnClicked();
-	virtual void NativeOnClicked();
+	COMMONUI_API void BP_OnClicked();
+	COMMONUI_API virtual void NativeOnClicked();
 
 	UFUNCTION(BlueprintImplementableEvent, Category = CommonButton, meta = (DisplayName = "On Double Clicked"))
-	void BP_OnDoubleClicked();
-	virtual void NativeOnDoubleClicked();
+	COMMONUI_API void BP_OnDoubleClicked();
+	COMMONUI_API virtual void NativeOnDoubleClicked();
 
 	UFUNCTION(BlueprintImplementableEvent, Category = CommonButton, meta = (DisplayName = "On Input Action Triggered"))
-	void BP_OnInputActionTriggered();
+	COMMONUI_API void BP_OnInputActionTriggered();
 
 	/** Unless this is called, we will assume the double click should be converted into a normal click. */
 	UFUNCTION(BlueprintCallable, Category = CommonButton)
-	void StopDoubleClickPropagation();
+	COMMONUI_API void StopDoubleClickPropagation();
 	
 	UFUNCTION(BlueprintImplementableEvent, Category = CommonButton, meta = (DisplayName = "On Pressed"))
-	void BP_OnPressed();
-	virtual void NativeOnPressed();
+	COMMONUI_API void BP_OnPressed();
+	COMMONUI_API virtual void NativeOnPressed();
 
 	UFUNCTION(BlueprintImplementableEvent, Category = CommonButton, meta = (DisplayName = "On Released"))
-	void BP_OnReleased();
-	virtual void NativeOnReleased();
+	COMMONUI_API void BP_OnReleased();
+	COMMONUI_API virtual void NativeOnReleased();
 
 	UFUNCTION(BlueprintImplementableEvent, Category = CommonButton, meta = (DisplayName = "On Enabled"))
-	void BP_OnEnabled();
-	virtual void NativeOnEnabled();
+	COMMONUI_API void BP_OnEnabled();
+	COMMONUI_API virtual void NativeOnEnabled();
 
 	UFUNCTION(BlueprintImplementableEvent, Category = CommonButton, meta = (DisplayName = "On Disabled"))
-	void BP_OnDisabled();
-	virtual void NativeOnDisabled();
+	COMMONUI_API void BP_OnDisabled();
+	COMMONUI_API virtual void NativeOnDisabled();
 
 	UFUNCTION(BlueprintImplementableEvent, Category = CommonButton, meta = (DisplayName = "On Input Method Changed"))
-	void BP_OnInputMethodChanged(ECommonInputType CurrentInputType);
+	COMMONUI_API void BP_OnInputMethodChanged(ECommonInputType CurrentInputType);
 
 	/** Allows derived classes to take action when the current text style has changed */
 	UFUNCTION(BlueprintImplementableEvent, meta=(BlueprintProtected="true"), Category = "Common Button")
-	void OnCurrentTextStyleChanged();
-	virtual void NativeOnCurrentTextStyleChanged();
+	COMMONUI_API void OnCurrentTextStyleChanged();
+	COMMONUI_API virtual void NativeOnCurrentTextStyleChanged();
 
 	/** Internal method to allow the selected state to be set regardless of selectability or toggleability */
 	UFUNCTION(BlueprintCallable, meta=(BlueprintProtected="true"), Category = "Common Button")
-	void SetSelectedInternal(bool bInSelected, bool bAllowSound = true, bool bBroadcast = true);
+	COMMONUI_API void SetSelectedInternal(bool bInSelected, bool bAllowSound = true, bool bBroadcast = true);
 
 	/** Callback fired when input action datatable row changes */
 	UFUNCTION(BlueprintImplementableEvent, Category = "Common Button")
-	void OnTriggeredInputActionChanged(const FDataTableRowHandle& NewTriggeredAction);
+	COMMONUI_API void OnTriggeredInputActionChanged(const FDataTableRowHandle& NewTriggeredAction);
 
 	/** Callback fired when triggered input action datatable row changes */
 	UFUNCTION(BlueprintImplementableEvent, Category = "Common Button")
-	void OnTriggeringInputActionChanged(const FDataTableRowHandle& NewTriggeredAction);
+	COMMONUI_API void OnTriggeringInputActionChanged(const FDataTableRowHandle& NewTriggeredAction);
 
 	/** Callback fired when enhanced input action changes */
 	UFUNCTION(BlueprintImplementableEvent, Category = "Common Button")
-	void OnTriggeringEnhancedInputActionChanged(const UInputAction* InInputAction);
+	COMMONUI_API void OnTriggeringEnhancedInputActionChanged(const UInputAction* InInputAction);
 
 	/** Callback fired continously during hold interactions */
 	UFUNCTION(BlueprintImplementableEvent, Category = "Common Button")
-	void OnActionProgress(float HeldPercent);
+	COMMONUI_API void OnActionProgress(float HeldPercent);
 	
 	/**
     * By default, if bRequiresHold is true, the bound Input Action will be forced to require hold as well.
     * The Bound Input action will use the row's hold values regardless of the row's bRequiresHold state.
     */
 	UFUNCTION()
-	virtual bool GetConvertInputActionToHold();
+	COMMONUI_API virtual bool GetConvertInputActionToHold();
     
     /** Bound to the hold progress of the bound key from the input action */
 	UFUNCTION()
-	virtual void NativeOnActionProgress(float HeldPercent);
+	COMMONUI_API virtual void NativeOnActionProgress(float HeldPercent);
 
 	/** Bound to the hold progress not related to the bound key */
 	UFUNCTION()
-	virtual bool NativeOnHoldProgress(float DeltaTime);
+	COMMONUI_API virtual bool NativeOnHoldProgress(float DeltaTime);
 	
 	/** Bound to the hold progress rollback not related to the bound key */
 	UFUNCTION()
-	virtual bool NativeOnHoldProgressRollback(float DeltaTime);
+	COMMONUI_API virtual bool NativeOnHoldProgressRollback(float DeltaTime);
 	
 	/** Callback fired when hold events complete */
 	UFUNCTION(BlueprintImplementableEvent, Category = "Common Button")
-	void OnActionComplete();
+	COMMONUI_API void OnActionComplete();
 
 	UFUNCTION()
-	virtual void NativeOnActionComplete();
+	COMMONUI_API virtual void NativeOnActionComplete();
 	
 	UFUNCTION()
-    virtual void HoldReset();
+    COMMONUI_API virtual void HoldReset();
 
-	virtual bool GetButtonAnalyticInfo(FString& ButtonName, FString& ABTestName, FString& ExtraData) const;
+	COMMONUI_API virtual bool GetButtonAnalyticInfo(FString& ButtonName, FString& ABTestName, FString& ExtraData) const;
 
-	void RefreshDimensions();
-	virtual void NativeOnMouseEnter( const FGeometry& InGeometry, const FPointerEvent& InMouseEvent ) override;
-	virtual void NativeOnMouseLeave( const FPointerEvent& InMouseEvent ) override;
+	COMMONUI_API void RefreshDimensions();
+	COMMONUI_API virtual void NativeOnMouseEnter( const FGeometry& InGeometry, const FPointerEvent& InMouseEvent ) override;
+	COMMONUI_API virtual void NativeOnMouseLeave( const FPointerEvent& InMouseEvent ) override;
 
-	void UpdateInputActionWidgetVisibility();
+	COMMONUI_API void UpdateInputActionWidgetVisibility();
 
 	/** The minimum width of the button (only used if greater than the style's minimum) */
 	UPROPERTY(EditAnywhere, BlueprintReadOnly, Category = Layout, meta = (ClampMin = "0"))
@@ -991,27 +990,27 @@ protected:
 };
 
 UCLASS(Transient)
-class COMMONUI_API UWidgetLockedStateRegistration : public UWidgetBinaryStateRegistration
+class UWidgetLockedStateRegistration : public UWidgetBinaryStateRegistration
 {
 	GENERATED_BODY()
 
 public:
 
 	/** Post-load initialized bit corresponding to this binary state */
-	static inline FWidgetStateBitfield Bit;
+	COMMONUI_API static inline FWidgetStateBitfield Bit;
 
-	static const inline FName StateName = FName("Locked");
+	COMMONUI_API static const inline FName StateName = FName("Locked");
 
 	//~ Begin UWidgetBinaryStateRegistration Interface.
-	virtual FName GetStateName() const override;
-	virtual bool GetRegisteredWidgetState(const UWidget* InWidget) const override;
+	COMMONUI_API virtual FName GetStateName() const override;
+	COMMONUI_API virtual bool GetRegisteredWidgetState(const UWidget* InWidget) const override;
 	//~ End UWidgetBinaryStateRegistration Interface
 
 protected:
 	friend UWidgetStateSettings;
 
 	//~ Begin UWidgetBinaryStateRegistration Interface.
-	virtual void InitializeStaticBitfields() const override;
+	COMMONUI_API virtual void InitializeStaticBitfields() const override;
 	//~ End UWidgetBinaryStateRegistration Interface
 };
 
diff --git a/Engine/Plugins/Runtime/DataRegistry/Source/DataRegistry/Public/SoftDataRegistryOrTable.h b/Engine/Plugins/Runtime/DataRegistry/Source/DataRegistry/Public/SoftDataRegistryOrTable.h
--- a/Engine/Plugins/Runtime/DataRegistry/Source/DataRegistry/Public/SoftDataRegistryOrTable.h
+++ b/Engine/Plugins/Runtime/DataRegistry/Source/DataRegistry/Public/SoftDataRegistryOrTable.h
@@ -13,26 +13,26 @@
  * Defines a DataRegistry or a DataTable with a common interface to both
  */
 USTRUCT()
-struct DATAREGISTRY_API FSoftDataRegistryOrTable
+struct FSoftDataRegistryOrTable
 {
 	GENERATED_BODY()
 
-	FSoftDataRegistryOrTable();
-	FSoftDataRegistryOrTable(const UDataTable* InDataTable, const FDataRegistryType& InRegistryType);
+	DATAREGISTRY_API FSoftDataRegistryOrTable();
+	DATAREGISTRY_API FSoftDataRegistryOrTable(const UDataTable* InDataTable, const FDataRegistryType& InRegistryType);
 
-	bool Serialize(FArchive& Ar);
+	DATAREGISTRY_API bool Serialize(FArchive& Ar);
 
 	/* Method to check if the given table matches the value in this struct */
-	bool Matches(const UDataTable* InTable) const;
+	DATAREGISTRY_API bool Matches(const UDataTable* InTable) const;
 
 	/* Method to check if the given registry matches the value in this struct */
-	bool Matches(const UDataRegistry* InRegistry) const;
+	DATAREGISTRY_API bool Matches(const UDataRegistry* InRegistry) const;
 
 	/* Method to check if the given FDataRegistryOrTableRow matches the value in this struct */
-	bool Matches(const FDataRegistryOrTableRow& RegistryOrTableId) const;
+	DATAREGISTRY_API bool Matches(const FDataRegistryOrTableRow& RegistryOrTableId) const;
 
 	/* Method to check validity of this struct (Checks that the registry or data table is set)*/
-	bool IsValid() const;
+	DATAREGISTRY_API bool IsValid() const;
 
 	template <class T>
 	void ForEachItem(const FString& ContextString, TFunctionRef<void(const FName& Name, const T& Item)> Predicate) const
@@ -140,20 +140,20 @@ struct DATAREGISTRY_API FSoftDataRegistryOrTable
 	}
 
 	/** Used to upgrade a SoftObjectPtr to a FSoftDataRegistryOrTable */
-	bool SerializeFromMismatchedTag(const FPropertyTag& Tag, FStructuredArchive::FSlot Slot);
+	DATAREGISTRY_API bool SerializeFromMismatchedTag(const FPropertyTag& Tag, FStructuredArchive::FSlot Slot);
 
 	/** Returns whether or not the registry or table is loaded */
-	bool IsLoaded() const;
+	DATAREGISTRY_API bool IsLoaded() const;
 
 	/**
 	 * Requests an async load of a data table using the StreamableManager and then execute the callback, which will happen even if the load fails.
 	 * NOTE: Data registries cannot be loaded this way because they must be registered with the data registry subsystem.
 	 * @param DelegateToCall		Delegate to call when load finishes. Will be called on the next tick if asset is already loaded, or many seconds later
 	 */
-	void LoadAsync(FStreamableDelegate DelegateToCall);	
+	DATAREGISTRY_API void LoadAsync(FStreamableDelegate DelegateToCall);	
 
 	/** Method to get a FDataRegistryOrTableRow from the given row name */
-	FDataRegistryOrTableRow GetRegistryOrTableRow(FName RowName) const;
+	DATAREGISTRY_API FDataRegistryOrTableRow GetRegistryOrTableRow(FName RowName) const;
 
 	UPROPERTY(EditAnywhere, Category = DataRegistryOrTable)
 	bool bUseDataRegistry = false;
@@ -169,7 +169,7 @@ struct DATAREGISTRY_API FSoftDataRegistryOrTable
 private:
 
 	/* Method to get the Data Registry */
-	const UDataRegistry* GetDataRegistry() const;
+	DATAREGISTRY_API const UDataRegistry* GetDataRegistry() const;
 
 };
 
@@ -185,19 +185,19 @@ struct TStructOpsTypeTraits<FSoftDataRegistryOrTable> : public TStructOpsTypeTra
 
 /** Defines a DataRegistryId or DataTableRowHandle with a common interface to both */
 USTRUCT(BlueprintType)
-struct DATAREGISTRY_API FDataRegistryOrTableRow
+struct FDataRegistryOrTableRow
 {
 	GENERATED_BODY()
 
-	FDataRegistryOrTableRow();
-	FDataRegistryOrTableRow(const FDataTableRowHandle& RowHandle);
-	FDataRegistryOrTableRow(const FDataRegistryId& RegistryId);
+	DATAREGISTRY_API FDataRegistryOrTableRow();
+	DATAREGISTRY_API FDataRegistryOrTableRow(const FDataTableRowHandle& RowHandle);
+	DATAREGISTRY_API FDataRegistryOrTableRow(const FDataRegistryId& RegistryId);
 
 	/** Used to upgrade a FDataTableHandle to a FDataRegistryOrTableRow */
-	bool SerializeFromMismatchedTag(const FPropertyTag& Tag, FStructuredArchive::FSlot Slot);
+	DATAREGISTRY_API bool SerializeFromMismatchedTag(const FPropertyTag& Tag, FStructuredArchive::FSlot Slot);
 
 	/* Method to get the Data Registry for the given registry id */
-	const UDataRegistry* GetDataRegistry() const;
+	DATAREGISTRY_API const UDataRegistry* GetDataRegistry() const;
 
 	template <class T>
 	const T* GetItem(const TCHAR* ContextString) const
@@ -221,12 +221,12 @@ struct DATAREGISTRY_API FDataRegistryOrTableRow
 	 * Returns the script struct used for the data registry item or data table row. Only works for data registries that
 	 * are registered
 	 */
-	const UScriptStruct* GetStruct() const;
+	DATAREGISTRY_API const UScriptStruct* GetStruct() const;
 
-	FString ToString() const;
+	DATAREGISTRY_API FString ToString() const;
 
 	// method to check validity of this row
-	bool IsValid() const;
+	DATAREGISTRY_API bool IsValid() const;
 
 	UPROPERTY(EditAnywhere, Category = DataRegistryOrTable)
 	bool bUseDataRegistryId = false;
diff --git a/Engine/Source/Programs/UnrealBuildTool/ToolChain/ClangWarnings.cs b/Engine/Source/Programs/UnrealBuildTool/ToolChain/ClangWarnings.cs
--- a/Engine/Source/Programs/UnrealBuildTool/ToolChain/ClangWarnings.cs
+++ b/Engine/Source/Programs/UnrealBuildTool/ToolChain/ClangWarnings.cs
@@ -158,6 +158,12 @@ namespace UnrealBuildTool
 					Arguments.Add("-Wno-" + Warning);
 				}
 			}
+
+			if (CompileEnvironment.bIsBuildingDLL)
+			{
+				// The code base contains lots of places where we do "extern template X_API class ..." and we want to keep doing that to reduce compile times.
+				Arguments.Add("-Wno-dllexport-explicit-instantiation-decl");
+			}
 		}
 
 		internal static void GetHeaderDisabledWarnings(List<string> Arguments)
diff --git a/Engine/Source/Runtime/Core/Public/Experimental/Containers/HazardPointer.h b/Engine/Source/Runtime/Core/Public/Experimental/Containers/HazardPointer.h
--- a/Engine/Source/Runtime/Core/Public/Experimental/Containers/HazardPointer.h
+++ b/Engine/Source/Runtime/Core/Public/Experimental/Containers/HazardPointer.h
@@ -168,7 +168,7 @@ class FHazardPointerCollection
 	CORE_API void Delete(const HazardPointer_Impl::FHazardDeleter& Deleter, int32 CollectLimit);
 
 	template<bool Cached>
-	FHazardRecord* Grow();
+	CORE_API FHazardRecord* Grow();
 
 public:
 	FHazardPointerCollection()
diff --git a/Engine/Source/Runtime/Core/Public/GenericPlatform/GenericPlatformString.h b/Engine/Source/Runtime/Core/Public/GenericPlatform/GenericPlatformString.h
--- a/Engine/Source/Runtime/Core/Public/GenericPlatform/GenericPlatformString.h
+++ b/Engine/Source/Runtime/Core/Public/GenericPlatform/GenericPlatformString.h
@@ -419,7 +419,7 @@ private:
 	 * @param Src     Pointer to the null-terminated string being converted.
 	 */
 	template <typename DestEncoding, typename SourceEncoding>
-	static void LogBogusChars(const SourceEncoding* Src);
+	CORE_API static void LogBogusChars(const SourceEncoding* Src);
 
 
 	/**
@@ -429,5 +429,5 @@ private:
 	 * @param SrcSize Number of characters in the Src string.
 	 */
 	template <typename DestEncoding, typename SourceEncoding>
-	static void LogBogusChars(const SourceEncoding* Src, int32 SrcSize);
+	CORE_API static void LogBogusChars(const SourceEncoding* Src, int32 SrcSize);
 };
diff --git a/Engine/Source/Runtime/Core/Public/Misc/Base64.h b/Engine/Source/Runtime/Core/Public/Misc/Base64.h
--- a/Engine/Source/Runtime/Core/Public/Misc/Base64.h
+++ b/Engine/Source/Runtime/Core/Public/Misc/Base64.h
@@ -62,7 +62,8 @@ struct FBase64
 	 *
 	 * @return The length of the encoded data
 	 */
-	template<typename CharType> static uint32 Encode(const uint8* Source, uint32 Length, CharType* Dest, EBase64Mode Mode = EBase64Mode::Standard);
+	template<typename CharType>
+	CORE_API static uint32 Encode(const uint8* Source, uint32 Length, CharType* Dest, EBase64Mode Mode = EBase64Mode::Standard);
 
 	/**
 	* Get the encoded data size for the given number of bytes.
@@ -108,7 +109,8 @@ struct FBase64
 	 *
 	 * @return true if the buffer was decoded, false if it was invalid.
 	 */
-	template<typename CharType> static bool Decode(const CharType* Source, uint32 Length, uint8* Dest, EBase64Mode Mode = EBase64Mode::Standard);
+	template<typename CharType>
+	CORE_API static bool Decode(const CharType* Source, uint32 Length, uint8* Dest, EBase64Mode Mode = EBase64Mode::Standard);
 
 	/**
 	* Determine the decoded data size for the incoming base64 encoded string
@@ -127,7 +129,8 @@ struct FBase64
 	*
 	* @return The size in bytes of the decoded data
 	*/
-	template<typename CharType> static uint32 GetDecodedDataSize(const CharType* Source, uint32 Length);
+	template<typename CharType>
+	CORE_API static uint32 GetDecodedDataSize(const CharType* Source, uint32 Length);
 
 	/**
 	* Get the maximum decoded data size for the given number of input characters.
diff --git a/Engine/Source/Runtime/Experimental/Chaos/Private/Chaos/PBDCollisionSpringConstraintsBase.cpp b/Engine/Source/Runtime/Experimental/Chaos/Private/Chaos/PBDCollisionSpringConstraintsBase.cpp
--- a/Engine/Source/Runtime/Experimental/Chaos/Private/Chaos/PBDCollisionSpringConstraintsBase.cpp
+++ b/Engine/Source/Runtime/Experimental/Chaos/Private/Chaos/PBDCollisionSpringConstraintsBase.cpp
@@ -600,8 +600,8 @@ void FPBDCollisionSpringConstraintsBase::Apply(SolverParticlesOrRange& InParticl
 	ApplyDynamicConstraints(InParticles, Dt);
 	ApplyKinematicConstraints(InParticles, Dt);
 }
-template void CHAOS_API FPBDCollisionSpringConstraintsBase::Apply(FSolverParticles& Particles, const FSolverReal Dt) const;
-template void CHAOS_API FPBDCollisionSpringConstraintsBase::Apply(FSolverParticlesRange& Particles, const FSolverReal Dt) const;
+template CHAOS_API void FPBDCollisionSpringConstraintsBase::Apply(FSolverParticles& Particles, const FSolverReal Dt) const;
+template CHAOS_API void FPBDCollisionSpringConstraintsBase::Apply(FSolverParticlesRange& Particles, const FSolverReal Dt) const;
 
 template<typename SolverParticlesOrRange>
 void FPBDCollisionSpringConstraintsBase::ApplyDynamicConstraints(SolverParticlesOrRange& Particles, const FSolverReal Dt) const
diff --git a/Engine/Source/Runtime/Experimental/Chaos/Public/Chaos/BoundingVolume.h b/Engine/Source/Runtime/Experimental/Chaos/Public/Chaos/BoundingVolume.h
--- a/Engine/Source/Runtime/Experimental/Chaos/Public/Chaos/BoundingVolume.h
+++ b/Engine/Source/Runtime/Experimental/Chaos/Public/Chaos/BoundingVolume.h
@@ -1376,7 +1376,7 @@ FArchive& operator<<(FChaosArchive& Ar, TBoundingVolume<TPayloadType, T, d>& Bou
 	return Ar;
 }
 
-#if PLATFORM_MAC || PLATFORM_LINUX
+#if PLATFORM_COMPILER_CLANG
 extern template class CHAOS_API Chaos::TBoundingVolume<int32, Chaos::FReal, 3>;
 extern template class CHAOS_API Chaos::TBoundingVolume<Chaos::FAccelerationStructureHandle, Chaos::FReal, 3>;
 #else
diff --git a/Engine/Source/Runtime/Experimental/Chaos/Public/Chaos/CollectionPropertyFacade.h b/Engine/Source/Runtime/Experimental/Chaos/Public/Chaos/CollectionPropertyFacade.h
--- a/Engine/Source/Runtime/Experimental/Chaos/Public/Chaos/CollectionPropertyFacade.h
+++ b/Engine/Source/Runtime/Experimental/Chaos/Public/Chaos/CollectionPropertyFacade.h
@@ -205,10 +205,10 @@ namespace Chaos::Softs
 		}
 
 		template<typename T, typename ElementType>
-		T GetValue(int32 KeyIndex, const TConstArrayView<ElementType>& ValueArray) const;
+		CHAOS_API T GetValue(int32 KeyIndex, const TConstArrayView<ElementType>& ValueArray) const;
 
 		template <typename T>
-		TConstArrayView<T> GetArray(const FName& Name) const;
+		CHAOS_API TConstArrayView<T> GetArray(const FName& Name) const;
 
 		bool HasAnyFlags(int32 KeyIndex, ECollectionPropertyFlags Flags) const { return EnumHasAnyFlags(GetValue<ECollectionPropertyFlags, uint8>(KeyIndex, FlagsArray), Flags); }
 
diff --git a/Engine/Source/Runtime/Experimental/Chaos/Public/Chaos/ISpatialAcceleration.h b/Engine/Source/Runtime/Experimental/Chaos/Public/Chaos/ISpatialAcceleration.h
--- a/Engine/Source/Runtime/Experimental/Chaos/Public/Chaos/ISpatialAcceleration.h
+++ b/Engine/Source/Runtime/Experimental/Chaos/Public/Chaos/ISpatialAcceleration.h
@@ -1020,7 +1020,7 @@ FORCEINLINE bool PrePreFilterHelper(const int32 Payload, const TVisitor& Visitor
 	return false;
 }
 
-#if PLATFORM_MAC || PLATFORM_LINUX
+#if PLATFORM_COMPILER_CLANG
 extern template class CHAOS_API ISpatialAcceleration<int32, FReal, 3>;
 extern template class CHAOS_API ISpatialVisitor<int32, FReal>;
 #else
diff --git a/Engine/Source/Runtime/Experimental/Chaos/Public/Chaos/PBDCollisionSpringConstraintsBase.h b/Engine/Source/Runtime/Experimental/Chaos/Public/Chaos/PBDCollisionSpringConstraintsBase.h
--- a/Engine/Source/Runtime/Experimental/Chaos/Public/Chaos/PBDCollisionSpringConstraintsBase.h
+++ b/Engine/Source/Runtime/Experimental/Chaos/Public/Chaos/PBDCollisionSpringConstraintsBase.h
@@ -68,10 +68,10 @@ public:
 
 	template<typename SpatialAccelerator, typename SolverParticlesOrRange>
 	UE_DEPRECATED(5.4, "Use Init with CollidableSubMesh. This method is much less efficient as it recreates the CollidableSubMesh each call.")
-	void Init(const SolverParticlesOrRange& Particles, const SpatialAccelerator& Spatial, const TConstArrayView<FPBDTriangleMeshCollisions::FGIAColor>& VertexGIAColors, const TArray<FPBDTriangleMeshCollisions::FGIAColor>& TriangleGIAColors);
+	CHAOS_API void Init(const SolverParticlesOrRange& Particles, const SpatialAccelerator& Spatial, const TConstArrayView<FPBDTriangleMeshCollisions::FGIAColor>& VertexGIAColors, const TArray<FPBDTriangleMeshCollisions::FGIAColor>& TriangleGIAColors);
 
 	template<typename SpatialAccelerator, typename SolverParticlesOrRange>
-	void Init(const SolverParticlesOrRange& Particles, const FSolverReal Dt, const FPBDTriangleMeshCollisions::FTriangleSubMesh& CollidableSubMesh,
+	CHAOS_API void Init(const SolverParticlesOrRange& Particles, const FSolverReal Dt, const FPBDTriangleMeshCollisions::FTriangleSubMesh& CollidableSubMesh,
 		const SpatialAccelerator& DynamicSpatial, const SpatialAccelerator& KinematicColliderSpatial, const TConstArrayView<FPBDTriangleMeshCollisions::FGIAColor>& VertexGIAColors, const TArray<FPBDTriangleMeshCollisions::FGIAColor>& TriangleGIAColors);
 
 	template<typename SolverParticlesOrRange>
@@ -147,7 +147,7 @@ public:
 	}
 
 	template<typename SolverParticlesOrRange>
-	void Apply(SolverParticlesOrRange& InParticles, const FSolverReal Dt) const;
+	CHAOS_API void Apply(SolverParticlesOrRange& InParticles, const FSolverReal Dt) const;
 
 	void Apply(FSolverParticles& InParticles, const FSolverReal Dt, const TArray<int32>& InConstraintIndices) const
 	{
diff --git a/Engine/Source/Runtime/Experimental/Chaos/Public/Chaos/ParticleHandle.h b/Engine/Source/Runtime/Experimental/Chaos/Public/Chaos/ParticleHandle.h
--- a/Engine/Source/Runtime/Experimental/Chaos/Public/Chaos/ParticleHandle.h
+++ b/Engine/Source/Runtime/Experimental/Chaos/Public/Chaos/ParticleHandle.h
@@ -3916,7 +3916,7 @@ inline void SetObjectStateHelper(IPhysicsProxyBase& Proxy, FPBDRigidParticle& Ri
 
 CHAOS_API void SetObjectStateHelper(IPhysicsProxyBase& Proxy, FPBDRigidParticleHandle& Rigid, EObjectStateType InState, bool bAllowEvents = false, bool bInvalidate = true);
 
-#if PLATFORM_MAC || PLATFORM_LINUX
+#if PLATFORM_COMPILER_CLANG
 extern template class CHAOS_API ISpatialAcceleration<FAccelerationStructureHandle, FReal, 3>;
 extern template class CHAOS_API ISpatialVisitor<FAccelerationStructureHandle, FReal>;
 #else
diff --git a/Engine/Source/Runtime/Experimental/Chaos/Public/Chaos/TriangleMesh.h b/Engine/Source/Runtime/Experimental/Chaos/Public/Chaos/TriangleMesh.h
--- a/Engine/Source/Runtime/Experimental/Chaos/Public/Chaos/TriangleMesh.h
+++ b/Engine/Source/Runtime/Experimental/Chaos/Public/Chaos/TriangleMesh.h
@@ -152,9 +152,9 @@ namespace Chaos
 			If this is not the case the returned face normals may be inverted
 		*/
 		template <typename T>
-		TArray<TVec3<T>> GetFaceNormals(const TConstArrayView<TVec3<T>>& Points, const bool ReturnEmptyOnError = true) const;
+		CHAOS_API TArray<TVec3<T>> GetFaceNormals(const TConstArrayView<TVec3<T>>& Points, const bool ReturnEmptyOnError = true) const;
 		template <typename T>
-		void GetFaceNormals(TArray<TVec3<T>>& Normals, const TConstArrayView<TVec3<T>>& Points, const bool ReturnEmptyOnError = true) const;
+		CHAOS_API void GetFaceNormals(TArray<TVec3<T>>& Normals, const TConstArrayView<TVec3<T>>& Points, const bool ReturnEmptyOnError = true) const;
 		FORCEINLINE TArray<FVec3> GetFaceNormals(const FParticles& InParticles, const bool ReturnEmptyOnError = true) const
 		{ return GetFaceNormals(TConstArrayView<FVec3>(InParticles.X()), ReturnEmptyOnError); }
 
diff --git a/Engine/Source/Runtime/GeometryCore/Private/DynamicMesh/DynamicMeshOverlay.cpp b/Engine/Source/Runtime/GeometryCore/Private/DynamicMesh/DynamicMeshOverlay.cpp
--- a/Engine/Source/Runtime/GeometryCore/Private/DynamicMesh/DynamicMeshOverlay.cpp
+++ b/Engine/Source/Runtime/GeometryCore/Private/DynamicMesh/DynamicMeshOverlay.cpp
@@ -1859,29 +1859,37 @@ namespace UE
 namespace Geometry
 {
 
+#if PLATFORM_COMPILER_CLANG
+#define UE_EXTERN_TEMPLATE_API GEOMETRYCORE_API
+#else
+#define UE_EXTERN_TEMPLATE_API
+#endif
+
 // These are explicit instantiations of the templates that are exported from the shared lib.
 // Only these instantiations of the template can be used.
 // This is necessary because we have placed most of the templated functions in this .cpp file, instead of the header.
-template class GEOMETRYCORE_API TDynamicMeshOverlay<float, 1>;
-template class GEOMETRYCORE_API TDynamicMeshOverlay<double, 1>;
-template class GEOMETRYCORE_API TDynamicMeshOverlay<int, 1>;
-template class GEOMETRYCORE_API TDynamicMeshOverlay<float, 2>;
-template class GEOMETRYCORE_API TDynamicMeshOverlay<double, 2>;
-template class GEOMETRYCORE_API TDynamicMeshOverlay<int, 2>;
-template class GEOMETRYCORE_API TDynamicMeshOverlay<float, 3>;
-template class GEOMETRYCORE_API TDynamicMeshOverlay<double, 3>;
-template class GEOMETRYCORE_API TDynamicMeshOverlay<int, 3>;
-template class GEOMETRYCORE_API TDynamicMeshOverlay<float, 4>;
-template class GEOMETRYCORE_API TDynamicMeshOverlay<double, 4>;
-
-template class GEOMETRYCORE_API TDynamicMeshVectorOverlay<float, 2, FVector2f>;
-template class GEOMETRYCORE_API TDynamicMeshVectorOverlay<double, 2, FVector2d>;
-template class GEOMETRYCORE_API TDynamicMeshVectorOverlay<int, 2, FVector2i>;
-template class GEOMETRYCORE_API TDynamicMeshVectorOverlay<float, 3, FVector3f>;
-template class GEOMETRYCORE_API TDynamicMeshVectorOverlay<double, 3, FVector3d>;
-template class GEOMETRYCORE_API TDynamicMeshVectorOverlay<int, 3, FVector3i>;
-template class GEOMETRYCORE_API TDynamicMeshVectorOverlay<float, 4, FVector4f>;
-template class GEOMETRYCORE_API TDynamicMeshVectorOverlay<double, 4, FVector4d>;
+template class UE_EXTERN_TEMPLATE_API TDynamicMeshOverlay<float, 1>;
+template class UE_EXTERN_TEMPLATE_API TDynamicMeshOverlay<double, 1>;
+template class UE_EXTERN_TEMPLATE_API TDynamicMeshOverlay<int, 1>;
+template class UE_EXTERN_TEMPLATE_API TDynamicMeshOverlay<float, 2>;
+template class UE_EXTERN_TEMPLATE_API TDynamicMeshOverlay<double, 2>;
+template class UE_EXTERN_TEMPLATE_API TDynamicMeshOverlay<int, 2>;
+template class UE_EXTERN_TEMPLATE_API TDynamicMeshOverlay<float, 3>;
+template class UE_EXTERN_TEMPLATE_API TDynamicMeshOverlay<double, 3>;
+template class UE_EXTERN_TEMPLATE_API TDynamicMeshOverlay<int, 3>;
+template class UE_EXTERN_TEMPLATE_API TDynamicMeshOverlay<float, 4>;
+template class UE_EXTERN_TEMPLATE_API TDynamicMeshOverlay<double, 4>;
+
+template class UE_EXTERN_TEMPLATE_API TDynamicMeshVectorOverlay<float, 2, FVector2f>;
+template class UE_EXTERN_TEMPLATE_API TDynamicMeshVectorOverlay<double, 2, FVector2d>;
+template class UE_EXTERN_TEMPLATE_API TDynamicMeshVectorOverlay<int, 2, FVector2i>;
+template class UE_EXTERN_TEMPLATE_API TDynamicMeshVectorOverlay<float, 3, FVector3f>;
+template class UE_EXTERN_TEMPLATE_API TDynamicMeshVectorOverlay<double, 3, FVector3d>;
+template class UE_EXTERN_TEMPLATE_API TDynamicMeshVectorOverlay<int, 3, FVector3i>;
+template class UE_EXTERN_TEMPLATE_API TDynamicMeshVectorOverlay<float, 4, FVector4f>;
+template class UE_EXTERN_TEMPLATE_API TDynamicMeshVectorOverlay<double, 4, FVector4d>;
+
+#undef UE_EXTERN_TEMPLATE_API
 
 } // end namespace UE::Geometry
 } // end namespace UE
\ No newline at end of file
diff --git a/Engine/Source/Runtime/GeometryCore/Private/DynamicMesh/MeshTangents.cpp b/Engine/Source/Runtime/GeometryCore/Private/DynamicMesh/MeshTangents.cpp
--- a/Engine/Source/Runtime/GeometryCore/Private/DynamicMesh/MeshTangents.cpp
+++ b/Engine/Source/Runtime/GeometryCore/Private/DynamicMesh/MeshTangents.cpp
@@ -673,8 +673,16 @@ namespace UE
 namespace Geometry
 {
 
-template class GEOMETRYCORE_API TMeshTangents<float>;
-template class GEOMETRYCORE_API TMeshTangents<double>;
+#if PLATFORM_COMPILER_CLANG
+#define UE_EXTERN_TEMPLATE_API GEOMETRYCORE_API
+#else
+#define UE_EXTERN_TEMPLATE_API
+#endif
+
+template class UE_EXTERN_TEMPLATE_API TMeshTangents<float>;
+template class UE_EXTERN_TEMPLATE_API TMeshTangents<double>;
+
+#undef UE_EXTERN_TEMPLATE_API
 
 } // end namespace UE::Geometry
 } // end namespace UE
diff --git a/Engine/Source/Runtime/GeometryCore/Public/DynamicMesh/DynamicMeshOverlay.h b/Engine/Source/Runtime/GeometryCore/Public/DynamicMesh/DynamicMeshOverlay.h
--- a/Engine/Source/Runtime/GeometryCore/Public/DynamicMesh/DynamicMeshOverlay.h
+++ b/Engine/Source/Runtime/GeometryCore/Public/DynamicMesh/DynamicMeshOverlay.h
@@ -248,7 +248,7 @@ public:
 	}
 
 	/** Discard all elements. */
-	void ClearElements();
+	GEOMETRYCORE_API void ClearElements();
 
 	/** Discard elements for given triangles. */
 	template <typename EnumerableIntType>
@@ -278,9 +278,9 @@ public:
 
 
 	/** Allocate a new element with the given constant value */
-	int AppendElement(RealType ConstantValue);
+	GEOMETRYCORE_API int AppendElement(RealType ConstantValue);
 	/** Allocate a new element with the given value */
-	int AppendElement(const RealType* Value);
+	GEOMETRYCORE_API int AppendElement(const RealType* Value);
 
 	void SetParentVertex(int ElementIndex, int ParentVertexIndex)
 	{
@@ -288,7 +288,7 @@ public:
 	}
 
 	/** Initialize the triangle list to the given size, and set all triangles to InvalidID */
-	void InitializeTriangles(int MaxTriangleID);
+	GEOMETRYCORE_API void InitializeTriangles(int MaxTriangleID);
 
 	/**
 	 * Set the triangle to the given Element index tuple, and increment element reference counts 
@@ -298,7 +298,7 @@ public:
 	 *  when remeshing across existing elements to avoid them being freed while temporarily unreferenced,
 	 *  but then it should eventually be followed by a call to FreeUnusedElements().
 	 */
-	EMeshResult SetTriangle(int TriangleID, const FIndex3i& TriElements, bool bAllowElementFreeing = true);
+	GEOMETRYCORE_API EMeshResult SetTriangle(int TriangleID, const FIndex3i& TriElements, bool bAllowElementFreeing = true);
 
 	/**
 	 * Goes through elements and frees any whose reference counts indicate that they are not being used. This
@@ -307,7 +307,7 @@ public:
 	 * 
 	 * @param ElementsToCheck If provided, only these element ID's will be checked.
 	 */
-	void FreeUnusedElements(const TSet<int>* ElementsToCheck = nullptr);
+	GEOMETRYCORE_API void FreeUnusedElements(const TSet<int>* ElementsToCheck = nullptr);
 
 	/**
 	 * Set the triangle to have InvalidID element IDs, decrementing element reference counts if needed.
@@ -316,7 +316,7 @@ public:
 	 *  become immediately unallocated. This can be set to false as part of a remeshing, but then it 
 	 *  should eventually be followed by a call to FreeUnusedElements().
 	 */
-	void UnsetTriangle(int TriangleID, bool bAllowElementFreeing = true);
+	GEOMETRYCORE_API void UnsetTriangle(int TriangleID, bool bAllowElementFreeing = true);
 
 	/** @return true if this triangle was set */
 	bool IsSetTriangle(int TID) const
@@ -336,7 +336,7 @@ public:
 	 *									  Note if a vertex can be shared between tris A and B, and B and C, it will be shared between all three
 	 * @param InitElementValue Initial element value, copied into all created elements
 	 */
-	void CreateFromPredicate(TFunctionRef<bool(int ParentVertexIdx, int TriIDA, int TriIDB)> TrisCanShareVertexPredicate, RealType InitElementValue);
+	GEOMETRYCORE_API void CreateFromPredicate(TFunctionRef<bool(int ParentVertexIdx, int TriIDA, int TriIDB)> TrisCanShareVertexPredicate, RealType InitElementValue);
 
 	/**
 	 * Refine an existing overlay topology.  For any element on a given triangle, if the predicate returns true, it gets topologically split out so it isn't shared by any other triangle.
@@ -345,7 +345,7 @@ public:
 	 * @param ShouldSplitOutVertex predicate returns true of the element should be split out and not shared w/ any other triangle
 	 * @param GetNewElementValue function to assign a new value to any element that is split out
 	 */
-	void SplitVerticesWithPredicate(TFunctionRef<bool(int ElementIdx, int TriID)> ShouldSplitOutVertex, TFunctionRef<void(int ElementIdx, int TriID, RealType* FillVect)> GetNewElementValue);
+	GEOMETRYCORE_API void SplitVerticesWithPredicate(TFunctionRef<bool(int ElementIdx, int TriID)> ShouldSplitOutVertex, TFunctionRef<void(int ElementIdx, int TriID, RealType* FillVect)> GetNewElementValue);
 
 	/**
 	 * Collapse SourceElementID into TargetElementID, resulting in connecting any containing triangles and reducing the total elements in the overlay.
@@ -354,7 +354,7 @@ public:
 	 * @param TargetElementID the element to merge into
 	 * @return If the operation completed successfully, returns true
 	 */
-	bool MergeElement(int SourceElementID, int TargetElementID);
+	GEOMETRYCORE_API bool MergeElement(int SourceElementID, int TargetElementID);
 
 	/**
 	 * Create a new copy of ElementID, and update connected triangles in the TrianglesToUpdate array to reference the copy of ElementID where they used to reference ElementID
@@ -364,7 +364,7 @@ public:
 	 * @param TrianglesToUpdate the triangles that should now reference the new element
 	 * @return the ID of the newly created element
 	 */
-	int SplitElement(int ElementID, const TArrayView<const int>& TrianglesToUpdate);
+	GEOMETRYCORE_API int SplitElement(int ElementID, const TArrayView<const int>& TrianglesToUpdate);
 
 	/**
 	* Create a new copy of ElementID, and update connected triangles in the TrianglesToUpdate array to reference the copy of ElementID where they used to reference ElementID.  The new element will have the given parent vertex ID.
@@ -375,7 +375,7 @@ public:
 	* @param TrianglesToUpdate the triangles that should now reference the new element.  Note: this is allowed to include triangles that do not have the element at all; sometimes you may want to do so to avoid creating a new array for each call.
 	* @return the ID of the newly created element
 	*/
-	int SplitElementWithNewParent(int ElementID, int SplitParentVertexID, const TArrayView<const int>& TrianglesToUpdate);
+	GEOMETRYCORE_API int SplitElementWithNewParent(int ElementID, int SplitParentVertexID, const TArrayView<const int>& TrianglesToUpdate);
 
 	/**
 	 * Split any bowties at given vertex.
@@ -383,12 +383,12 @@ public:
 	 * @param NewElementIDs If not null, newly created element IDs are placed here. Note that this array is
 	 *   intentionally not cleared before appending to it.
 	 */
-	void SplitBowtiesAtVertex(int32 Vid, TArray<int32>* NewElementIDs = nullptr);
+	GEOMETRYCORE_API void SplitBowtiesAtVertex(int32 Vid, TArray<int32>* NewElementIDs = nullptr);
 
 	/**
 	* Refine an existing overlay topology by splitting any bow ties
 	*/
-	void SplitBowties();
+	GEOMETRYCORE_API void SplitBowties();
 
 
 	//
@@ -416,7 +416,7 @@ public:
 	 * If bUnsafe, we use fast id allocation that does not update free list.
 	 * You should only be using this between BeginUnsafeElementsInsert() / EndUnsafeElementsInsert() calls
 	 */
-	EMeshResult InsertElement(int ElementID, const RealType* Value, bool bUnsafe = false);
+	GEOMETRYCORE_API EMeshResult InsertElement(int ElementID, const RealType* Value, bool bUnsafe = false);
 
 
 	//
@@ -525,39 +525,39 @@ public:
 	*   If present, bIsNonIntersectingOut will be true only if this is a seam edge 
 	*   that does not intersect with another seam or the end of the seam. 
 	*/
-	bool IsSeamEdge(int EdgeID, bool* bIsNonIntersectingOut = nullptr) const;
+	GEOMETRYCORE_API bool IsSeamEdge(int EdgeID, bool* bIsNonIntersectingOut = nullptr) const;
 	/** Returns true if the parent-mesh edge is a "Seam End" in this overlay, meaning the adjacent element triangles share one element, not two */
-	bool IsSeamEndEdge(int EdgeID) const;
+	GEOMETRYCORE_API bool IsSeamEndEdge(int EdgeID) const;
 	/** Returns true if the parent-mesh vertex is connected to any seam edges */
-	bool IsSeamVertex(int VertexID, bool bBoundaryIsSeam = true) const;
+	GEOMETRYCORE_API bool IsSeamVertex(int VertexID, bool bBoundaryIsSeam = true) const;
 
 	/**
 	 * Determines whether the base-mesh vertex has "bowtie" topology in the Overlay.
 	 * Bowtie topology means that one or more elements at the vertex are shared across disconnected UV-components.
 	 * @return true if the base-mesh vertex has "bowtie" topology in the overlay
 	 */
-	bool IsBowtieInOverlay(int32 VertexID) const;
+	GEOMETRYCORE_API bool IsBowtieInOverlay(int32 VertexID) const;
 
 	/** @return true if the two triangles are connected, ie shared edge exists and is not a seam edge */
-	bool AreTrianglesConnected(int TriangleID0, int TriangleID1) const;
+	GEOMETRYCORE_API bool AreTrianglesConnected(int TriangleID0, int TriangleID1) const;
 
 	/** find the elements associated with a given parent-mesh vertex */
-	void GetVertexElements(int VertexID, TArray<int>& OutElements) const;
+	GEOMETRYCORE_API void GetVertexElements(int VertexID, TArray<int>& OutElements) const;
 	/** Count the number of unique elements for a given parent-mesh vertex */
-	int CountVertexElements(int VertexID, bool bBruteForce = false) const;
+	GEOMETRYCORE_API int CountVertexElements(int VertexID, bool bBruteForce = false) const;
 
 	/** find the triangles connected to an element */
-	void GetElementTriangles(int ElementID, TArray<int>& OutTriangles) const;
+	GEOMETRYCORE_API void GetElementTriangles(int ElementID, TArray<int>& OutTriangles) const;
 
 	/** 
 	 * Find the element ID at a vertex of a triangle. 
 	 * 
 	 * @return Returns the element ID or FDynamicMesh3::InvalidID if the vertex is not a parent of any element contained in the triangle. 
 	 */
-	int GetElementIDAtVertex(int TriangleID, int VertexID) const;
+	GEOMETRYCORE_API int GetElementIDAtVertex(int TriangleID, int VertexID) const;
 
 	/** @return true if overlay has any interior seam edges. This requires an O(N) search unless it early-outs. */
-	bool HasInteriorSeamEdges() const;
+	GEOMETRYCORE_API bool HasInteriorSeamEdges() const;
 
 	/**
 	 * Compute interpolated parameter value inside triangle using barycentric coordinates
@@ -582,12 +582,12 @@ public:
 	/**
 	 * Checks that the overlay mesh is well-formed, ie all internal data structures are consistent
 	 */
-	bool CheckValidity(bool bAllowNonManifoldVertices = true, EValidityCheckFailMode FailMode = EValidityCheckFailMode::Check) const;
+	GEOMETRYCORE_API bool CheckValidity(bool bAllowNonManifoldVertices = true, EValidityCheckFailMode FailMode = EValidityCheckFailMode::Check) const;
 
 	/**
 	 * Returns true if this overlay is the same as Other.
 	 */
-	bool IsSameAs(const TDynamicMeshOverlay<RealType, ElementSize>& Other, bool bIgnoreDataLayout) const;
+	GEOMETRYCORE_API bool IsSameAs(const TDynamicMeshOverlay<RealType, ElementSize>& Other, bool bIgnoreDataLayout) const;
 
 	/**
 	 * Serialization operator for FDynamicMeshOverlay.
@@ -609,7 +609,7 @@ public:
 	 * @param CompactMaps If this is not a null pointer, the mesh serialization compacted the vertex and/or triangle data using the provided mapping. 
 	 * @param bUseCompression Use compression for serializing bulk data.
 	 */
-	void Serialize(FArchive& Ar, const FCompactMaps* CompactMaps, bool bUseCompression);
+	GEOMETRYCORE_API void Serialize(FArchive& Ar, const FCompactMaps* CompactMaps, bool bUseCompression);
 
 public:
 	/** Set a triangle's element indices to InvalidID */
@@ -760,7 +760,7 @@ public:
 	 * @param bFindUniqueElements if true, ProcessFunc is only called once for each ElementID, otherwise it is called once for each Triangle
 	 * @return true if at least one valid Element was found, ie if ProcessFunc was called at least one time
 	 */
-	bool EnumerateVertexElements(
+	GEOMETRYCORE_API bool EnumerateVertexElements(
 		int VertexID,
 		TFunctionRef<bool(int TriangleID, int ElementID, const VectorType& Value)> ProcessFunc,
 		bool bFindUniqueElements = true) const;
diff --git a/Engine/Source/Runtime/GeometryCore/Public/DynamicMesh/MeshTangents.h b/Engine/Source/Runtime/GeometryCore/Public/DynamicMesh/MeshTangents.h
--- a/Engine/Source/Runtime/GeometryCore/Public/DynamicMesh/MeshTangents.h
+++ b/Engine/Source/Runtime/GeometryCore/Public/DynamicMesh/MeshTangents.h
@@ -131,7 +131,7 @@ public:
 	 * Initialize Tangents from the FDynamicMeshAttributeSet of SourceMesh 
 	 * @return false if tangent attributes do not exist
 	 */
-	bool CopyTriVertexTangents(const FDynamicMesh3& SourceMesh);
+	GEOMETRYCORE_API bool CopyTriVertexTangents(const FDynamicMesh3& SourceMesh);
 
 
 	/**
@@ -139,7 +139,7 @@ public:
 	 * In this mode there is no averaging of tangents across triangles. So if we have N triangles
 	 * in the mesh, then 3*N tangents are generated. These tangents are computed in parallel.
 	 */
-	void ComputeTriVertexTangents(const FDynamicMeshNormalOverlay* NormalOverlay, const FDynamicMeshUVOverlay* UVOverlay, const FComputeTangentsOptions& Options);
+	GEOMETRYCORE_API void ComputeTriVertexTangents(const FDynamicMeshNormalOverlay* NormalOverlay, const FDynamicMeshUVOverlay* UVOverlay, const FComputeTangentsOptions& Options);
 
 	/**
 	 * Return tangent and bitangent at a vertex of triangle for per-triangle computed tangents
@@ -241,14 +241,14 @@ public:
 	 * Compute per-triangle tangents for the given UV Overlay 
 	 * @param bOrthogonalize if true, we orthogonalize the tangents by computing the Bitangent sign and regenerating it. Otherwise the tangents are left non-orthogonal.
 	 */
-	void ComputeTriangleTangents(const FDynamicMeshUVOverlay* UVOverlay, bool bOrthogonalize = true);
+	GEOMETRYCORE_API void ComputeTriangleTangents(const FDynamicMeshUVOverlay* UVOverlay, bool bOrthogonalize = true);
 
 
 	/**
 	 * Set Tangents on mesh overlays
 	 * @param MeshToSet Mesh to copy overlays to; does not need to be the same as the Mesh member of this class
 	 */
-	bool CopyToOverlays(FDynamicMesh3& MeshToSet) const;
+	GEOMETRYCORE_API bool CopyToOverlays(FDynamicMesh3& MeshToSet) const;
 
 
 
@@ -280,7 +280,7 @@ protected:
 	/**
 	 * Set the size of the Tangents array to Count, and optionally clear all values to (0,0,0)
 	 */
-	void SetTangentCount(int Count, bool bClearToZero);
+	GEOMETRYCORE_API void SetTangentCount(int Count, bool bClearToZero);
 
 	// calculate per-triangle tangents and then projected to overlay normals at each triangle-vertex
 	void ComputeSeparatePerTriangleTangents(const FDynamicMeshNormalOverlay* NormalOverlay, const FDynamicMeshUVOverlay* UVOverlay, const FComputeTangentsOptions& Options);
diff --git a/Engine/Source/Runtime/GeometryCore/Public/Intersection/ContainmentQueries3.h b/Engine/Source/Runtime/GeometryCore/Public/Intersection/ContainmentQueries3.h
--- a/Engine/Source/Runtime/GeometryCore/Public/Intersection/ContainmentQueries3.h
+++ b/Engine/Source/Runtime/GeometryCore/Public/Intersection/ContainmentQueries3.h
@@ -34,15 +34,15 @@ namespace UE
 
 		/** @return true if InnerSphere is fully contained within OuterSphere */
 		template<typename RealType>
-		bool IsInside(const TSphere3<RealType>& OuterSphere, const TSphere3<RealType>& InnerSphere);
+		GEOMETRYCORE_API bool IsInside(const TSphere3<RealType>& OuterSphere, const TSphere3<RealType>& InnerSphere);
 
 		/** @return true if InnerCapsule is fully contained within OuterSphere */
 		template<typename RealType>
-		bool IsInside(const TSphere3<RealType>& OuterSphere, const TCapsule3<RealType>& InnerCapsule);
+		GEOMETRYCORE_API bool IsInside(const TSphere3<RealType>& OuterSphere, const TCapsule3<RealType>& InnerCapsule);
 
 		/** @return true if InnerBox is fully contained within OuterSphere */
 		template<typename RealType>
-		bool IsInside(const TSphere3<RealType>& OuterSphere, const TOrientedBox3<RealType>& InnerBox);
+		GEOMETRYCORE_API bool IsInside(const TSphere3<RealType>& OuterSphere, const TOrientedBox3<RealType>& InnerBox);
 
 		/** @return true if all all points in range-based for over EnumerablePts are inside OuterSphere */
 		template<typename RealType, typename EnumerablePointsType, typename E = decltype(DeclVal<EnumerablePointsType>().begin())>
@@ -73,15 +73,15 @@ namespace UE
 
 		/** @return true if InnerCapsule is fully contained within OuterCapsule */
 		template<typename RealType>
-		bool IsInside(const TCapsule3<RealType>& OuterCapsule, const TCapsule3<RealType>& InnerCapsule);
+		GEOMETRYCORE_API bool IsInside(const TCapsule3<RealType>& OuterCapsule, const TCapsule3<RealType>& InnerCapsule);
 
 		/** @return true if InnerSphere is fully contained within OuterCapsule */
 		template<typename RealType>
-		bool IsInside(const TCapsule3<RealType>& OuterCapsule, const TSphere3<RealType>& InnerSphere);
+		GEOMETRYCORE_API bool IsInside(const TCapsule3<RealType>& OuterCapsule, const TSphere3<RealType>& InnerSphere);
 
 		/** @return true if InnerBox is fully contained within OuterCapsule */
 		template<typename RealType>
-		bool IsInside(const TCapsule3<RealType>& OuterCapsule, const TOrientedBox3<RealType>& InnerBox);
+		GEOMETRYCORE_API bool IsInside(const TCapsule3<RealType>& OuterCapsule, const TOrientedBox3<RealType>& InnerBox);
 
 		/** @return true if all all points in range-based for over EnumerablePts are inside OuterCapsule */
 		template<typename RealType, typename EnumerablePointsType, typename E = decltype(DeclVal<EnumerablePointsType>().begin())>
@@ -112,15 +112,15 @@ namespace UE
 
 		/** @return true if InnerBox is fully contained within OuterBox */
 		template<typename RealType>
-		bool IsInside(const TOrientedBox3<RealType>& OuterBox, const TOrientedBox3<RealType>& InnerBox);
+		GEOMETRYCORE_API bool IsInside(const TOrientedBox3<RealType>& OuterBox, const TOrientedBox3<RealType>& InnerBox);
 
 		/** @return true if InnerSphere is fully contained within OuterBox */
 		template<typename RealType>
-		bool IsInside(const TOrientedBox3<RealType>& OuterBox, const TSphere3<RealType>& InnerSphere);
+		GEOMETRYCORE_API bool IsInside(const TOrientedBox3<RealType>& OuterBox, const TSphere3<RealType>& InnerSphere);
 
 		/** @return true if InnerCapsule is fully contained within OuterBox */
 		template<typename RealType>
-		bool IsInside(const TOrientedBox3<RealType>& OuterBox, const TCapsule3<RealType>& InnerCapsule);
+		GEOMETRYCORE_API bool IsInside(const TOrientedBox3<RealType>& OuterBox, const TCapsule3<RealType>& InnerCapsule);
 
 		/** @return true if all all points in range-based for over EnumerablePts are inside OuterBox */
 		template<typename RealType, typename EnumerablePointsType, typename E = decltype(DeclVal<EnumerablePointsType>().begin())>
diff --git a/Engine/Source/Runtime/GeometryCore/Public/Intersection/IntersectionQueries2.h b/Engine/Source/Runtime/GeometryCore/Public/Intersection/IntersectionQueries2.h
--- a/Engine/Source/Runtime/GeometryCore/Public/Intersection/IntersectionQueries2.h
+++ b/Engine/Source/Runtime/GeometryCore/Public/Intersection/IntersectionQueries2.h
@@ -15,7 +15,7 @@ template <typename T> struct TSegment2;
 
 // Return true if Segment and Box intersect and false otherwise
 template<typename Real>
-bool TestIntersection(const TSegment2<Real>& Segment, const TAxisAlignedBox2<Real>& Box);
+GEOMETRYCORE_API bool TestIntersection(const TSegment2<Real>& Segment, const TAxisAlignedBox2<Real>& Box);
 		
 } // namespace UE::Geometry
 } // namespace UE
diff --git a/Engine/Source/Runtime/GeometryCore/Public/Intersection/IntersectionQueries3.h b/Engine/Source/Runtime/GeometryCore/Public/Intersection/IntersectionQueries3.h
--- a/Engine/Source/Runtime/GeometryCore/Public/Intersection/IntersectionQueries3.h
+++ b/Engine/Source/Runtime/GeometryCore/Public/Intersection/IntersectionQueries3.h
@@ -23,15 +23,15 @@ namespace UE
 
 		/** @return true if Halfspace and Sphere intersect */
 		template<typename RealType>
-		bool TestIntersection(const THalfspace3<RealType>& Halfspace, const TSphere3<RealType>& Sphere);
+		GEOMETRYCORE_API bool TestIntersection(const THalfspace3<RealType>& Halfspace, const TSphere3<RealType>& Sphere);
 
 		/** @return true if Halfspace and Capsule intersect */
 		template<typename RealType>
-		bool TestIntersection(const THalfspace3<RealType>& Halfspace, const TCapsule3<RealType>& Capsule);
+		GEOMETRYCORE_API bool TestIntersection(const THalfspace3<RealType>& Halfspace, const TCapsule3<RealType>& Capsule);
 
 		/** @return true if Halfspace and Box intersect */
 		template<typename RealType>
-		bool TestIntersection(const THalfspace3<RealType>& Halfspace, const TOrientedBox3<RealType>& Box);
+		GEOMETRYCORE_API bool TestIntersection(const THalfspace3<RealType>& Halfspace, const TOrientedBox3<RealType>& Box);
 
 
 	}
diff --git a/Engine/Source/Runtime/GeometryCore/Public/MeshRegionBoundaryLoops.h b/Engine/Source/Runtime/GeometryCore/Public/MeshRegionBoundaryLoops.h
--- a/Engine/Source/Runtime/GeometryCore/Public/MeshRegionBoundaryLoops.h
+++ b/Engine/Source/Runtime/GeometryCore/Public/MeshRegionBoundaryLoops.h
@@ -91,7 +91,7 @@ public:
 	 * @returns false if there is an error (if the overlay does not have an element for one of the vids)
 	 */
 	template<typename StorageType, int ElementSize, typename ElementType>
-	bool GetLoopOverlayMap(const FEdgeLoop& LoopIn,
+	GEOMETRYCORE_API bool GetLoopOverlayMap(const FEdgeLoop& LoopIn,
 		const TDynamicMeshOverlay<StorageType, ElementSize>& Overlay,
 		VidOverlayMap<ElementType>& LoopVidsToOverlayElementsOut);
 
@@ -102,7 +102,7 @@ public:
 	 * and remove any UV elements that no longer exist (because the vertex was on a UV seam).
 	 */
 	template<typename StorageType, int ElementSize, typename ElementType>
-	void UpdateLoopOverlayMapValidity(VidOverlayMap<ElementType>& LoopVidsToOverlayElements,
+	GEOMETRYCORE_API void UpdateLoopOverlayMapValidity(VidOverlayMap<ElementType>& LoopVidsToOverlayElements,
 		const TDynamicMeshOverlay<StorageType, ElementSize>& Overlay);
 
 	/**
