Force input type to return as touch within the common input preprocessor when in the editor and using the UseMouseForTouch feature.
Prevents the game from switching to keyboard and mouse in certain situations

#tests PIE
#rnx

diff --git a/Engine/Plugins/Runtime/CommonUI/Source/CommonInput/Private/CommonInputPreprocessor.cpp b/Engine/Plugins/Runtime/CommonUI/Source/CommonInput/Private/CommonInputPreprocessor.cpp
--- a/Engine/Plugins/Runtime/CommonUI/Source/CommonInput/Private/CommonInputPreprocessor.cpp
+++ b/Engine/Plugins/Runtime/CommonUI/Source/CommonInput/Private/CommonInputPreprocessor.cpp
@@ -253,6 +253,10 @@ ECommonInputType FCommonInputPreprocessor::GetInputType(const FKey& Key)
 
 ECommonInputType FCommonInputPreprocessor::GetInputType(const FPointerEvent& PointerEvent)
 {
+#if WITH_EDITOR
+	const ULocalPlayer& LocalPlayer = *InputSubsystem.GetLocalPlayerChecked();
+#endif
+
 	if (PointerEvent.IsTouchEvent())
 	{
 		return ECommonInputType::Touch;
@@ -261,5 +265,11 @@ ECommonInputType FCommonInputPreprocessor::GetInputType(const FPointerEvent& Poi
 	{
 		return ECommonInputType::Gamepad;
 	}
+#if WITH_EDITOR
+	else if (LocalPlayer.ViewportClient && LocalPlayer.ViewportClient->GetUseMouseForTouch())
+	{
+		return ECommonInputType::Touch;
+	}
+#endif
 	return ECommonInputType::MouseAndKeyboard;
 }
