Always update icon even if design time key is not set, allowing us to clear the icon
Only set visibility when not in design time

#tests Editor

diff --git a/Engine/Plugins/Runtime/CommonUI/Source/CommonUI/Private/CommonActionWidget.cpp b/Engine/Plugins/Runtime/CommonUI/Source/CommonUI/Private/CommonActionWidget.cpp
--- a/Engine/Plugins/Runtime/CommonUI/Source/CommonUI/Private/CommonActionWidget.cpp
+++ b/Engine/Plugins/Runtime/CommonUI/Source/CommonUI/Private/CommonActionWidget.cpp
@@ -328,7 +328,10 @@ void UCommonActionWidget::UpdateActionWidget()
 
 				if (Icon.DrawAs == ESlateBrushDrawType::NoDrawType)
 				{
-					SetVisibility(ESlateVisibility::Collapsed);
+					if (!IsDesignTime())
+					{
+						SetVisibility(ESlateVisibility::Collapsed);
+					}
 				}
 				else if (MyIcon.IsValid())
 				{
@@ -352,14 +355,21 @@ void UCommonActionWidget::UpdateActionWidget()
 					}
 
 					MyKeyBox->Invalidate(EInvalidateWidget::LayoutAndVolatility);
-					SetVisibility(IsDesignTime() ? ESlateVisibility::Visible : ESlateVisibility::SelfHitTestInvisible);
+
+					if (!IsDesignTime())
+					{
+						SetVisibility(ESlateVisibility::SelfHitTestInvisible);
+					}
 
 					return;
 				}
 			}
 		}
 
-		SetVisibility(ESlateVisibility::Collapsed);
+		if (!IsDesignTime())
+		{
+			SetVisibility(ESlateVisibility::Collapsed);
+		}
 	}
 }
 
@@ -374,7 +384,7 @@ bool UCommonActionWidget::ShouldUpdateActionWidgetIcon() const
 	const bool bIsEnhancedInputAction = EnhancedInputAction && CommonUI::IsEnhancedInputSupportEnabled();
 
 #if WITH_EDITORONLY_DATA
-	const bool bIsDesignPreview = IsDesignTime() && DesignTimeKey.IsValid();
+	const bool bIsDesignPreview = IsDesignTime();
 #else
 	const bool bIsDesignPreview = false;
 #endif
