PR #11415: Fix for Shipping Build Console with Common UI Input.

diff --git a/Engine/Plugins/Runtime/CommonUI/Source/CommonUI/Private/CommonGameViewportClient.cpp b/Engine/Plugins/Runtime/CommonUI/Source/CommonUI/Private/CommonGameViewportClient.cpp
--- a/Engine/Plugins/Runtime/CommonUI/Source/CommonUI/Private/CommonGameViewportClient.cpp
+++ b/Engine/Plugins/Runtime/CommonUI/Source/CommonUI/Private/CommonGameViewportClient.cpp
@@ -45,7 +45,7 @@ bool UCommonGameViewportClient::InputKey(const FInputKeyEventArgs& InEventArgs)
 	}
 
 	// The input is fair game for handling - the UI gets first dibs
-#if !UE_BUILD_SHIPPING
+#if ALLOW_CONSOLE
 	if (ViewportConsole && !ViewportConsole->ConsoleState.IsEqual(NAME_Typing) && !ViewportConsole->ConsoleState.IsEqual(NAME_Open))
 #endif
 	{		
@@ -82,7 +82,7 @@ bool UCommonGameViewportClient::InputAxis(FViewport* InViewport, FInputDeviceId
 
 bool UCommonGameViewportClient::InputTouch(FViewport* InViewport, int32 ControllerId, uint32 Handle, ETouchType::Type Type, const FVector2D& TouchLocation, float Force, FDateTime DeviceTimestamp, uint32 TouchpadIndex)
 {
-#if !UE_BUILD_SHIPPING
+#if ALLOW_CONSOLE
 	if (ViewportConsole != NULL && (ViewportConsole->ConsoleState != NAME_Typing) && (ViewportConsole->ConsoleState != NAME_Open))
 #endif
 	{
@@ -223,7 +223,7 @@ void UCommonGameViewportClient::HandleRerouteTouch(int32 ControllerId, uint32 To
 
 bool UCommonGameViewportClient::IsKeyPriorityAboveUI(const FInputKeyEventArgs& EventArgs)
 {
-#if !UE_BUILD_SHIPPING
+#if ALLOW_CONSOLE
 	// First priority goes to the viewport console regardless any state or setting
 	if (ViewportConsole && ViewportConsole->InputKey(EventArgs.InputDevice, EventArgs.Key, EventArgs.Event, EventArgs.AmountDepressed, EventArgs.IsGamepad()))
 	{
