PR #11854: Fix: common button styling consistently updates via SetIsInteractionEnabled
PR #11267: CommonUI: Fix CommonButtonBase not binding Pressed/Released events when using Enhanced Input

diff --git a/Engine/Plugins/Runtime/CommonUI/Source/CommonUI/Private/CommonButtonBase.cpp b/Engine/Plugins/Runtime/CommonUI/Source/CommonUI/Private/CommonButtonBase.cpp
--- a/Engine/Plugins/Runtime/CommonUI/Source/CommonUI/Private/CommonButtonBase.cpp
+++ b/Engine/Plugins/Runtime/CommonUI/Source/CommonUI/Private/CommonButtonBase.cpp
@@ -587,6 +587,8 @@ void UCommonButtonBase::BindTriggeringInputActionToClick()
 	{
 		FBindUIActionArgs BindArgs(TriggeringEnhancedInputAction, false, FSimpleDelegate::CreateUObject(this, &UCommonButtonBase::HandleTriggeringActionCommited));
 		BindArgs.OnHoldActionProgressed.BindUObject(this, &UCommonButtonBase::NativeOnActionProgress);
+		BindArgs.OnHoldActionPressed.BindUObject(this, &UCommonButtonBase::NativeOnPressed);
+		BindArgs.OnHoldActionReleased.BindUObject(this, &UCommonButtonBase::NativeOnReleased);
 		BindArgs.bIsPersistent = bIsPersistentBinding;
 
 		BindArgs.InputMode = InputModeOverride;
@@ -714,6 +716,7 @@ void UCommonButtonBase::SetIsInteractionEnabled(bool bInIsInteractionEnabled)
 			NativeOnUnhovered();
 		}
 	}
+	SetButtonStyle();
 }
 
 void UCommonButtonBase::SetHideInputAction(bool bInHideInputAction)
@@ -1950,4 +1953,4 @@ bool UWidgetLockedStateRegistration::GetRegisteredWidgetState(const UWidget* InW
 void UWidgetLockedStateRegistration::InitializeStaticBitfields() const
 {
 	Bit = FWidgetStateBitfield(GetStateName());
-}
\ No newline at end of file
+}
