PR #11416: Expose CommonVideoPlayer functions to BP

diff --git a/Engine/Plugins/Runtime/CommonUI/Source/CommonUI/Public/CommonVideoPlayer.h b/Engine/Plugins/Runtime/CommonUI/Source/CommonUI/Public/CommonVideoPlayer.h
--- a/Engine/Plugins/Runtime/CommonUI/Source/CommonUI/Public/CommonVideoPlayer.h
+++ b/Engine/Plugins/Runtime/CommonUI/Source/CommonUI/Public/CommonVideoPlayer.h
@@ -25,27 +25,45 @@ public:
 	UCommonVideoPlayer(const FObjectInitializer& Initializer);
 	virtual void PostInitProperties() override;
 
+	UFUNCTION(BlueprintCallable, Category="Video Player")
 	void SetVideo(UMediaSource* NewVideo);
+	UFUNCTION(BlueprintCallable, Category="Video Player")
 	void Seek(float PlaybackTime);
+	UFUNCTION(BlueprintCallable, Category="Video Player")
 	void Close();
 
+	UFUNCTION(BlueprintCallable, Category="Video Player")
 	void SetPlaybackRate(float PlaybackRate);
+	UFUNCTION(BlueprintCallable, Category="Video Player")
 	void SetLooping(bool bShouldLoopPlayback);
+	UFUNCTION(BlueprintCallable, Category="Video Player")
 	void SetIsMuted(bool bInIsMuted);
+	UFUNCTION(BlueprintCallable, Category="Video Player")
 	void SetShouldMatchSize(bool bInMatchSize);
 
+	UFUNCTION(BlueprintCallable, Category="Video Player")
 	void Play();
+	UFUNCTION(BlueprintCallable, Category="Video Player")
 	void Reverse();
+	UFUNCTION(BlueprintCallable, Category="Video Player")
 	void Pause();
+	UFUNCTION(BlueprintCallable, Category="Video Player")
 	void PlayFromStart();
 
+	UFUNCTION(BlueprintCallable, Category="Video Player")
 	float GetVideoDuration() const;
+	UFUNCTION(BlueprintCallable, Category="Video Player")
 	float GetPlaybackTime() const;
+	UFUNCTION(BlueprintCallable, Category="Video Player")
 	float GetPlaybackRate() const;
 
+	UFUNCTION(BlueprintCallable, Category="Video Player")
 	bool IsLooping() const;
+	UFUNCTION(BlueprintCallable, Category="Video Player")
 	bool IsPaused() const;
+	UFUNCTION(BlueprintCallable, Category="Video Player")
 	bool IsPlaying() const;
+	UFUNCTION(BlueprintCallable, Category="Video Player")
 	bool IsMuted() const { return bIsMuted; }
 
 	FSimpleMulticastDelegate& OnPlaybackResumed() { return OnPlaybackResumedEvent; }
