PR #11266: CommonUI: Adding virtual to methods that were needed to support Enhanced Input's Triggers and Modifiers in CommonUI widgets

diff --git a/Engine/Plugins/Runtime/CommonUI/Source/CommonUI/Public/CommonActionWidget.h b/Engine/Plugins/Runtime/CommonUI/Source/CommonUI/Public/CommonActionWidget.h
--- a/Engine/Plugins/Runtime/CommonUI/Source/CommonUI/Public/CommonActionWidget.h
+++ b/Engine/Plugins/Runtime/CommonUI/Source/CommonUI/Public/CommonActionWidget.h
@@ -61,7 +61,7 @@ public:
 	void SetIconRimBrush(FSlateBrush InIconRimBrush);
 
 	UFUNCTION(BlueprintCallable, Category = CommonActionWidget)
-	bool IsHeldAction() const;
+	virtual bool IsHeldAction() const;
 
 	DECLARE_DYNAMIC_MULTICAST_DELEGATE_OneParam(FOnInputMethodChanged, bool, bUsingGamepad);
 	UPROPERTY(BlueprintAssignable, Category = CommonActionWidget)
diff --git a/Engine/Plugins/Runtime/CommonUI/Source/CommonUI/Public/Input/CommonUIActionRouterBase.h b/Engine/Plugins/Runtime/CommonUI/Source/CommonUI/Public/Input/CommonUIActionRouterBase.h
--- a/Engine/Plugins/Runtime/CommonUI/Source/CommonUI/Public/Input/CommonUIActionRouterBase.h
+++ b/Engine/Plugins/Runtime/CommonUI/Source/CommonUI/Public/Input/CommonUIActionRouterBase.h
@@ -124,8 +124,8 @@ public:
 	void NotifyUserWidgetConstructed(const UCommonUserWidget& Widget);
 	void NotifyUserWidgetDestructed(const UCommonUserWidget& Widget);
 	
-	void AddBinding(FUIActionBindingHandle Binding);
-	void RemoveBinding(FUIActionBindingHandle Binding);
+	virtual void AddBinding(FUIActionBindingHandle Binding);
+	virtual void RemoveBinding(FUIActionBindingHandle Binding);
 
 	int32 GetLocalPlayerIndex() const;
 
