PR #11939: Fix Common Bound Action Bar

diff --git a/Engine/Plugins/Runtime/CommonUI/Source/CommonUI/Private/Input/CommonBoundActionBar.cpp b/Engine/Plugins/Runtime/CommonUI/Source/CommonUI/Private/Input/CommonBoundActionBar.cpp
--- a/Engine/Plugins/Runtime/CommonUI/Source/CommonUI/Private/Input/CommonBoundActionBar.cpp
+++ b/Engine/Plugins/Runtime/CommonUI/Source/CommonUI/Private/Input/CommonBoundActionBar.cpp
@@ -106,6 +106,11 @@ void UCommonBoundActionBar::ReleaseSlateResources(bool bReleaseChildren)
 			{
 				ActionRouter->OnBoundActionsUpdated().RemoveAll(this);
 			}
+
+			if (UCommonInputSubsystem* InputSubsystem = UCommonInputSubsystem::Get(LocalPlayer))
+			{
+				InputSubsystem->OnInputMethodChangedNative.RemoveAll(this);
+			}
 		}
 	}
 }
@@ -353,12 +358,29 @@ void UCommonBoundActionBar::HandlePlayerAdded(int32 PlayerIdx)
 	HandleBoundActionsUpdated(NewPlayer == GetOwningLocalPlayer());
 }
 
+void UCommonBoundActionBar::HandledInputTypeUpdated(ECommonInputType InputType)
+{
+	bIsRefreshQueued = true;
+}
+
 void UCommonBoundActionBar::MonitorPlayerActions(const ULocalPlayer* NewPlayer)
 {
 	if (const UCommonUIActionRouterBase* ActionRouter = ULocalPlayer::GetSubsystem<UCommonUIActionRouterBase>(NewPlayer))
 	{
 		ActionRouter->OnBoundActionsUpdated().AddUObject(this, &UCommonBoundActionBar::HandleBoundActionsUpdated, NewPlayer == GetOwningLocalPlayer());
 	}
+
+	/**
+	 * Update available inputs anytime the input type changes;
+	 * in Enhanced Input, an input action might be only bound with one device (gamepad), but not others (KBM / mouse),
+	 * due to that they will be considered invalid in case the action bar is updated when using the device that doesn't
+	 * have them bound to anything. When switching to a device that actually binds them to an input, the action bar
+	 * has to refresh to include the input actions that were considered invalid prior to that
+	 */
+	if (UCommonInputSubsystem* InputSubsystem = UCommonInputSubsystem::Get(NewPlayer))
+	{
+		InputSubsystem->OnInputMethodChangedNative.AddUObject(this, &UCommonBoundActionBar::HandledInputTypeUpdated);
+	}
 }
 
 void UCommonBoundActionBar::ActionBarUpdateBegin()
diff --git a/Engine/Plugins/Runtime/CommonUI/Source/CommonUI/Public/Input/CommonBoundActionBar.h b/Engine/Plugins/Runtime/CommonUI/Source/CommonUI/Public/Input/CommonBoundActionBar.h
--- a/Engine/Plugins/Runtime/CommonUI/Source/CommonUI/Public/Input/CommonBoundActionBar.h
+++ b/Engine/Plugins/Runtime/CommonUI/Source/CommonUI/Public/Input/CommonBoundActionBar.h
@@ -4,6 +4,7 @@
 
 #include "CommonButtonBase.h"
 #include "Components/DynamicEntryBoxBase.h"
+#include "CommonInputTypeEnum.h"
 #include "Tickable.h"
 #include "CommonBoundActionBar.generated.h"
 
@@ -54,6 +55,7 @@ protected:
 #endif
 
 private:
+	void HandledInputTypeUpdated(ECommonInputType InputType);
 	void HandleBoundActionsUpdated(bool bFromOwningPlayer);
 	void HandleDeferredDisplayUpdate();
 	void HandlePlayerAdded(int32 PlayerIdx);
