PR #12339: Add AlwaysSign to CommonNumericTextBlock

diff --git a/Engine/Plugins/Runtime/CommonUI/Source/CommonUI/Private/CommonNumericTextBlock.cpp b/Engine/Plugins/Runtime/CommonUI/Source/CommonUI/Private/CommonNumericTextBlock.cpp
--- a/Engine/Plugins/Runtime/CommonUI/Source/CommonUI/Private/CommonNumericTextBlock.cpp
+++ b/Engine/Plugins/Runtime/CommonUI/Source/CommonUI/Private/CommonNumericTextBlock.cpp
@@ -273,6 +273,7 @@ FNumberFormattingOptions UCommonNumericTextBlock::MakeNumberFormattingOptions()
 {
 	FNumberFormattingOptions NumberFormattingOptions;
 	NumberFormattingOptions.RoundingMode = FormattingSpecification.RoundingMode;
+	NumberFormattingOptions.AlwaysSign = FormattingSpecification.AlwaysSign;
 	NumberFormattingOptions.UseGrouping = FormattingSpecification.UseGrouping;
 	NumberFormattingOptions.MinimumIntegralDigits = FormattingSpecification.MinimumIntegralDigits;
 	NumberFormattingOptions.MaximumIntegralDigits = FormattingSpecification.MaximumIntegralDigits;
diff --git a/Engine/Plugins/Runtime/CommonUI/Source/CommonUI/Public/CommonNumericTextBlock.h b/Engine/Plugins/Runtime/CommonUI/Source/CommonUI/Public/CommonNumericTextBlock.h
--- a/Engine/Plugins/Runtime/CommonUI/Source/CommonUI/Public/CommonNumericTextBlock.h
+++ b/Engine/Plugins/Runtime/CommonUI/Source/CommonUI/Public/CommonNumericTextBlock.h
@@ -12,6 +12,7 @@ struct COMMONUI_API FCommonNumberFormattingOptions
 
 	FCommonNumberFormattingOptions()
 		: RoundingMode(ERoundingMode::HalfFromZero)
+		, AlwaysSign(false)
 		, UseGrouping(true)
 		, MinimumIntegralDigits(FNumberFormattingOptions::DefaultNoGrouping().MinimumIntegralDigits)
 		, MaximumIntegralDigits(FNumberFormattingOptions::DefaultNoGrouping().MaximumIntegralDigits)
@@ -22,7 +23,11 @@ struct COMMONUI_API FCommonNumberFormattingOptions
 	// The rounding mode to be used when the actual value can not be precisely represented due to restrictions on the number of integral or fractional digits. See values for details.
 	UPROPERTY(EditAnywhere, BlueprintReadOnly, Category = "Numeral Formating")
 	TEnumAsByte<ERoundingMode> RoundingMode;
-	
+
+	// Should the numerals always display the sign. IE: "+1"
+	UPROPERTY(EditAnywhere, BlueprintReadOnly, Category = "Numeral Formating")
+	bool AlwaysSign;
+
 	// Should the numerals use group separators. IE: "1,000,000"
 	UPROPERTY(EditAnywhere, BlueprintReadOnly, Category = "Numeral Formating")
 	bool UseGrouping;
