PR #12269: [Common UI] Restore cursor visibility when exiting PIE while using a gamepad

#tests Ran Lyra in PIE without the fix and used the front end menu to quit PIE and verified the cursor remained hidden. Applied the fix and repeated the test and verfied the cursor was visible.

diff --git a/Engine/Plugins/Runtime/CommonUI/Source/CommonUI/Private/Input/CommonUIActionRouterBase.cpp b/Engine/Plugins/Runtime/CommonUI/Source/CommonUI/Private/Input/CommonUIActionRouterBase.cpp
--- a/Engine/Plugins/Runtime/CommonUI/Source/CommonUI/Private/Input/CommonUIActionRouterBase.cpp
+++ b/Engine/Plugins/Runtime/CommonUI/Source/CommonUI/Private/Input/CommonUIActionRouterBase.cpp
@@ -303,12 +303,26 @@ void UCommonUIActionRouterBase::RegisterAnalogCursorTick()
 void UCommonUIActionRouterBase::Deinitialize()
 {
 	Super::Deinitialize();
-	
+
 	if (FSlateApplication::IsInitialized())
 	{
-		FSlateApplication::Get().OnFocusChanging().RemoveAll(this);
-		FSlateApplication::Get().UnregisterInputPreProcessor(AnalogCursor);
+		FSlateApplication& SlateApplication = FSlateApplication::Get();
+
+		SlateApplication.OnFocusChanging().RemoveAll(this);
+		SlateApplication.UnregisterInputPreProcessor(AnalogCursor);
+
+#if WITH_EDITOR
+		ULocalPlayer* LocalPlayer = GetLocalPlayer();
+		const bool bCursorUser = IsValid(LocalPlayer) && LocalPlayer->GetSlateUser() == SlateApplication.GetCursorUser();
+
+		if (bCursorUser)
+		{
+			// This restores cursor visibility when exiting PIE while using a gamepad.
+			SlateApplication.UsePlatformCursorForCursorUser(true);
+		}
+#endif
 	}
+
 	FTSTicker::GetCoreTicker().RemoveTicker(TickHandle);
 	SetActiveRoot(nullptr);
 	HeldKeys.Empty();
