Perf tanks when both players are switching tabs in the build menu

Release all active widgets from the widget pool in UCommonTabListWidgetBase::RemoveAllTabs.

#tests Navigated tabs in PIE and stepped through widget pool code, observing widgets being pulled from inactive list.
#rnx

diff --git a/Engine/Plugins/Runtime/CommonUI/Source/CommonUI/Private/CommonTabListWidgetBase.cpp b/Engine/Plugins/Runtime/CommonUI/Source/CommonUI/Private/CommonTabListWidgetBase.cpp
--- a/Engine/Plugins/Runtime/CommonUI/Source/CommonUI/Private/CommonTabListWidgetBase.cpp
+++ b/Engine/Plugins/Runtime/CommonUI/Source/CommonUI/Private/CommonTabListWidgetBase.cpp
@@ -177,6 +177,8 @@ void UCommonTabListWidgetBase::RemoveAllTabs()
 			OnTabButtonRemoval.Broadcast(Key, TabButton);
 		}
 	}
+
+	TabButtonWidgetPool.ReleaseAll();
 }
 
 int32 UCommonTabListWidgetBase::GetTabCount() const
