Don't allow adding of input locks if the current platform does not support that input
Add logic to settings record to not allow the input lock if the setting is on and there is no mouse connected
Add logic to the settings record to detect a device removal and reset the lock so it can be added or removed based on whether there is a mouse connected

#tests PIE, Client
[RN] Inclued fix for locking mouse input

diff --git a/Engine/Plugins/Runtime/CommonUI/Source/CommonInput/Private/CommonInputSubsystem.cpp b/Engine/Plugins/Runtime/CommonUI/Source/CommonInput/Private/CommonInputSubsystem.cpp
--- a/Engine/Plugins/Runtime/CommonUI/Source/CommonInput/Private/CommonInputSubsystem.cpp
+++ b/Engine/Plugins/Runtime/CommonUI/Source/CommonInput/Private/CommonInputSubsystem.cpp
@@ -145,7 +145,8 @@ bool UCommonInputSubsystem::GetInputTypeFilter(ECommonInputType InputType) const
 
 void UCommonInputSubsystem::AddOrRemoveInputTypeLock(FName InReason, ECommonInputType InInputType, bool bAddLock)
 {
-	if (bAddLock)
+	// Make sure the input is supported before locking it, otherwise remove it if it exists
+	if (bAddLock && PlatformSupportsInputType(InInputType))
 	{
 		ECommonInputType& CurrentValue = CurrentInputLocks.FindOrAdd(InReason);
 		CurrentValue = InInputType;
