Fix for the accessability mouse input lock not always applying properly
Fix the adding  / removing of an input lock in the common input so that it properly dispatches change event

#tests PIE, Client, Client
#rnx

diff --git a/Engine/Plugins/Runtime/CommonUI/Source/CommonInput/Private/CommonInputSubsystem.cpp b/Engine/Plugins/Runtime/CommonUI/Source/CommonInput/Private/CommonInputSubsystem.cpp
--- a/Engine/Plugins/Runtime/CommonUI/Source/CommonInput/Private/CommonInputSubsystem.cpp
+++ b/Engine/Plugins/Runtime/CommonUI/Source/CommonInput/Private/CommonInputSubsystem.cpp
@@ -179,10 +179,12 @@ void UCommonInputSubsystem::AddOrRemoveInputTypeLock(FName InReason, ECommonInpu
 		CurrentInputLock = (ECommonInputType)ComputedInputLock;
 	}
 
+	const ECommonInputType PreviousInput = CurrentInputType;
+
 	// If a lock was put in place, lock the current input type.
 	CurrentInputType = LockInput(LastInputType);
 
-	if (CurrentInputType != LastInputType)
+	if (CurrentInputType != PreviousInput)
 	{
 		BroadcastInputMethodChanged();
 	}
