CommonUI: Fix crash caused by [REDACTED_PERSON] with null Triggers

diff --git a/Engine/Plugins/Runtime/CommonUI/Source/CommonUI/Private/CommonActionWidget.cpp b/Engine/Plugins/Runtime/CommonUI/Source/CommonUI/Private/CommonActionWidget.cpp
--- a/Engine/Plugins/Runtime/CommonUI/Source/CommonUI/Private/CommonActionWidget.cpp
+++ b/Engine/Plugins/Runtime/CommonUI/Source/CommonUI/Private/CommonActionWidget.cpp
@@ -225,7 +225,7 @@ bool UCommonActionWidget::IsHeldAction() const
 	{
 		for (const TObjectPtr<UInputTrigger>& Trigger : EnhancedInputAction->Triggers)
 		{
-			if (EnumHasAnyFlags(Trigger->GetSupportedTriggerEvents(), ETriggerEventsSupported::Ongoing))
+			if (Trigger != nullptr && EnumHasAnyFlags(Trigger->GetSupportedTriggerEvents(), ETriggerEventsSupported::Ongoing))
 			{
 				return true;
 			}
