Split activation of input mapping context from methods in order to control them in derived classes of Activatable Widgets in Common UI

diff --git a/Engine/Plugins/Runtime/CommonUI/Source/CommonUI/Private/CommonActivatableWidget.cpp b/Engine/Plugins/Runtime/CommonUI/Source/CommonUI/Private/CommonActivatableWidget.cpp
--- a/Engine/Plugins/Runtime/CommonUI/Source/CommonUI/Private/CommonActivatableWidget.cpp
+++ b/Engine/Plugins/Runtime/CommonUI/Source/CommonUI/Private/CommonActivatableWidget.cpp
@@ -271,17 +271,7 @@ void UCommonActivatableWidget::NativeOnActivated()
 			UE_LOG(LogCommonUI, Verbose, TEXT("[%s] set visibility to [%s] on activation"), *GetName(), *StaticEnum<ESlateVisibility>()->GetDisplayValueAsText(ActivatedVisibility).ToString());
 		}
 
-		if (CommonUI::IsEnhancedInputSupportEnabled() && InputMapping)
-		{
-			if (const ULocalPlayer* LocalPlayer = GetOwningLocalPlayer())
-			{
-				if (UEnhancedInputLocalPlayerSubsystem* InputSystem = LocalPlayer->GetSubsystem<UEnhancedInputLocalPlayerSubsystem>())
-				{
-					InputSystem->AddMappingContext(InputMapping, InputMappingPriority);
-				}
-			}
-		}
-
+		ActivateMappingContext();
 		BP_OnActivated();
 		OnActivated().Broadcast();
 		BP_OnWidgetActivated.Broadcast();
@@ -298,16 +288,7 @@ void UCommonActivatableWidget::NativeOnDeactivated()
 			UE_LOG(LogCommonUI, Verbose, TEXT("[%s] set visibility to [%s] on deactivation"), *GetName(), *StaticEnum<ESlateVisibility>()->GetDisplayValueAsText(DeactivatedVisibility).ToString());
 		}
 
-		if (CommonUI::IsEnhancedInputSupportEnabled() && InputMapping)
-		{
-			if (const ULocalPlayer* LocalPlayer = GetOwningLocalPlayer())
-			{
-				if (UEnhancedInputLocalPlayerSubsystem* InputSystem = LocalPlayer->GetSubsystem<UEnhancedInputLocalPlayerSubsystem>())
-				{
-					InputSystem->RemoveMappingContext(InputMapping);
-				}
-			}
-		}
+		DeactivateMappingContext();
 
 		// Cancel any holds that were active
 		ClearActiveHoldInputs();
@@ -318,6 +299,34 @@ void UCommonActivatableWidget::NativeOnDeactivated()
 	}
 }
 
+void UCommonActivatableWidget::ActivateMappingContext()
+{
+	if (CommonUI::IsEnhancedInputSupportEnabled() && InputMapping)
+	{
+		if (const ULocalPlayer* LocalPlayer = GetOwningLocalPlayer())
+		{
+			if (UEnhancedInputLocalPlayerSubsystem* InputSystem = LocalPlayer->GetSubsystem<UEnhancedInputLocalPlayerSubsystem>())
+			{
+				InputSystem->AddMappingContext(InputMapping, InputMappingPriority);
+			}
+		}
+	}
+}
+
+void UCommonActivatableWidget::DeactivateMappingContext()
+{
+	if (CommonUI::IsEnhancedInputSupportEnabled() && InputMapping)
+	{
+		if (const ULocalPlayer* LocalPlayer = GetOwningLocalPlayer())
+		{
+			if (UEnhancedInputLocalPlayerSubsystem* InputSystem = LocalPlayer->GetSubsystem<UEnhancedInputLocalPlayerSubsystem>())
+			{
+				InputSystem->RemoveMappingContext(InputMapping);
+			}
+		}
+	}
+}
+
 bool UCommonActivatableWidget::NativeOnHandleBackAction()
 {
 	if (bIsBackHandler)
diff --git a/Engine/Plugins/Runtime/CommonUI/Source/CommonUI/Public/CommonActivatableWidget.h b/Engine/Plugins/Runtime/CommonUI/Source/CommonUI/Public/CommonActivatableWidget.h
--- a/Engine/Plugins/Runtime/CommonUI/Source/CommonUI/Public/CommonActivatableWidget.h
+++ b/Engine/Plugins/Runtime/CommonUI/Source/CommonUI/Public/CommonActivatableWidget.h
@@ -149,10 +149,12 @@ protected:
 	UFUNCTION(BlueprintImplementableEvent, Category = ActivatableWidget, meta = (DisplayName = "On Activated"))
 	void BP_OnActivated();
 	virtual void NativeOnActivated();
+	virtual void ActivateMappingContext();
 
 	UFUNCTION(BlueprintImplementableEvent, Category = ActivatableWidget, meta = (DisplayName = "On Deactivated"))
 	void BP_OnDeactivated();
 	virtual void NativeOnDeactivated();
+	virtual void DeactivateMappingContext();
 
 	/** 
 	 * Override in BP implementations to provide custom behavior when receiving a back action 
