[Common UI] Add Broadcast event for Input Config Changed
- Add Event for Common Action Router to broadcast  Input Config changes

Enables game-logic responses to player input config changed

#tests PIE backstage through gameplay, results, and return to backstage

diff --git a/Engine/Plugins/Runtime/CommonUI/Source/CommonUI/Private/Input/CommonUIActionRouterBase.cpp b/Engine/Plugins/Runtime/CommonUI/Source/CommonUI/Private/Input/CommonUIActionRouterBase.cpp
--- a/Engine/Plugins/Runtime/CommonUI/Source/CommonUI/Private/Input/CommonUIActionRouterBase.cpp
+++ b/Engine/Plugins/Runtime/CommonUI/Source/CommonUI/Private/Input/CommonUIActionRouterBase.cpp
@@ -1638,6 +1638,8 @@ void UCommonUIActionRouterBase::ApplyUIInputConfig(const FUIInputConfig& NewConf
 		{
 			OnActiveInputModeChanged().Broadcast(NewConfig.GetInputMode());
 		}
+
+		OnActiveInputConfigChanged().Broadcast(NewConfig);
 	}
 }
 
diff --git a/Engine/Plugins/Runtime/CommonUI/Source/CommonUI/Public/Input/CommonUIActionRouterBase.h b/Engine/Plugins/Runtime/CommonUI/Source/CommonUI/Public/Input/CommonUIActionRouterBase.h
--- a/Engine/Plugins/Runtime/CommonUI/Source/CommonUI/Public/Input/CommonUIActionRouterBase.h
+++ b/Engine/Plugins/Runtime/CommonUI/Source/CommonUI/Public/Input/CommonUIActionRouterBase.h
@@ -88,6 +88,9 @@ public:
 	FOnActiveInputModeChanged& OnActiveInputModeChanged() const { return OnActiveInputModeChangedEvent; }
 	ECommonInputMode GetActiveInputMode(ECommonInputMode DefaultInputMode = ECommonInputMode::All) const;
 	EMouseCaptureMode GetActiveMouseCaptureMode(EMouseCaptureMode DefaultMouseCapture = EMouseCaptureMode::NoCapture) const;
+	
+	DECLARE_EVENT_OneParam(UCommonUIActionRouterBase, FOnActiveInputConfigChanged, const FUIInputConfig);
+	FOnActiveInputConfigChanged& OnActiveInputConfigChanged() const { return OnActiveInputConfigChangedEvent; }
 
 	DECLARE_EVENT_OneParam(UCommonUIActionRouterBase, FOnActivationMetadataChanged, FActivationMetadata);
 	FOnActivationMetadataChanged& OnActivationMetadataChanged() const { return OnActivationMetadataChangedEvent; }
@@ -235,6 +238,7 @@ PRAGMA_ENABLE_DEPRECATION_WARNINGS
 	mutable FSimpleMulticastDelegate OnBoundActionsUpdatedEvent;
 	mutable FOnActiveInputModeChanged OnActiveInputModeChangedEvent;
 	mutable FOnActivationMetadataChanged OnActivationMetadataChangedEvent;
+	mutable FOnActiveInputConfigChanged OnActiveInputConfigChangedEvent;
 
 	friend class FActionRouterBindingCollection;
 	friend class FActivatableTreeNode;
