Update Action Context Router to gather the inputs from the action domains even if it has a valid root
Change the logic for breaking out of the loop of action domains to check the internal functions provided by the domain

#tests PIE

#rnx

diff --git a/Engine/Plugins/Runtime/CommonUI/Source/CommonUI/Private/Input/CommonUIActionRouterBase.cpp b/Engine/Plugins/Runtime/CommonUI/Source/CommonUI/Private/Input/CommonUIActionRouterBase.cpp
--- a/Engine/Plugins/Runtime/CommonUI/Source/CommonUI/Private/Input/CommonUIActionRouterBase.cpp
+++ b/Engine/Plugins/Runtime/CommonUI/Source/CommonUI/Private/Input/CommonUIActionRouterBase.cpp
@@ -418,26 +418,31 @@ TArray<FUIActionBindingHandle> UCommonUIActionRouterBase::GatherActiveBindings()
 	{
 		ActiveRootNode->AppendAllActiveActions(BindingHandles);
 	}
-	else if (const UCommonInputActionDomainTable* ActionDomainTable = GetActionDomainTable())
+	
+	if (const UCommonInputActionDomainTable* ActionDomainTable = GetActionDomainTable())
 	{
-		bool bFoundRoot = false;
+		bool bDomainHadActiveRoots = false;
+
 		for (const UCommonInputActionDomain* ActionDomain : ActionDomainTable->ActionDomains)
 		{
 			if (const FActionDomainSortedRootList* SortedRootList = ActionDomainRootNodes.Find(ActionDomain))
 			{
 				for (const FActivatableTreeRootRef& RootNode : SortedRootList->RootList)
 				{
-					// only return the first of the root nodes that's in the sorted root list in action domain order.
-					if (RootNode->IsReceivingInput() && RootNode->DoesWidgetSupportActivationFocus())
+					if (RootNode->IsReceivingInput() && RootNode->IsWidgetActivated())
 					{
 						RootNode->AppendAllActiveActions(BindingHandles);
-						bFoundRoot = true;
-						break;
+						bDomainHadActiveRoots = true;
+
+						if (ActionDomain->ShouldBreakInnerEventFlow(false))
+						{
+							break;
+						}
 					}
 				}
 			}
 
-			if (bFoundRoot)
+			if (ActionDomain && ActionDomain->ShouldBreakEventFlow(bDomainHadActiveRoots, false))
 			{
 				break;
 			}
