Add additional check for nullptr of ActionDomain

#tests PIE

#rnx

diff --git a/Engine/Plugins/Runtime/CommonUI/Source/CommonUI/Private/Input/CommonUIActionRouterBase.cpp b/Engine/Plugins/Runtime/CommonUI/Source/CommonUI/Private/Input/CommonUIActionRouterBase.cpp
--- a/Engine/Plugins/Runtime/CommonUI/Source/CommonUI/Private/Input/CommonUIActionRouterBase.cpp
+++ b/Engine/Plugins/Runtime/CommonUI/Source/CommonUI/Private/Input/CommonUIActionRouterBase.cpp
@@ -434,7 +434,7 @@ TArray<FUIActionBindingHandle> UCommonUIActionRouterBase::GatherActiveBindings()
 						RootNode->AppendAllActiveActions(BindingHandles);
 						bDomainHadActiveRoots = true;
 
-						if (ActionDomain->ShouldBreakInnerEventFlow(false))
+						if (ActionDomain && ActionDomain->ShouldBreakInnerEventFlow(false))
 						{
 							break;
 						}
