PR #12089: CommonUI - CommonButtonBase make Inpu handling functions virtual

diff --git a/Engine/Plugins/Runtime/CommonUI/Source/CommonUI/Public/CommonButtonBase.h b/Engine/Plugins/Runtime/CommonUI/Source/CommonUI/Public/CommonButtonBase.h
--- a/Engine/Plugins/Runtime/CommonUI/Source/CommonUI/Public/CommonButtonBase.h
+++ b/Engine/Plugins/Runtime/CommonUI/Source/CommonUI/Public/CommonButtonBase.h
@@ -551,10 +551,10 @@ protected:
 
 	/** Handler function registered to the underlying button's click. */
 	UFUNCTION()
-	COMMONUI_API void HandleButtonClicked();
+	COMMONUI_API virtual void HandleButtonClicked();
 
 	/** Handler function registered to the underlying button's double click. */
-	COMMONUI_API FReply HandleButtonDoubleClicked();
+	COMMONUI_API virtual FReply HandleButtonDoubleClicked();
 
 	/** Helper function registered to the underlying button receiving focus */
 	UFUNCTION()
@@ -566,11 +566,11 @@ protected:
 
 	/** Helper function registered to the underlying button when pressed */
 	UFUNCTION()
-	COMMONUI_API void HandleButtonPressed();
+	COMMONUI_API virtual void HandleButtonPressed();
 
 	/** Helper function registered to the underlying button when released */
 	UFUNCTION()
-	COMMONUI_API void HandleButtonReleased();
+	COMMONUI_API virtual void HandleButtonReleased();
 
 	UFUNCTION(BlueprintImplementableEvent, Category = CommonButton, meta = (DisplayName = "On Selected"))
 	COMMONUI_API void BP_OnSelected();
