CommonUI: Fix hitting a check in List View caused by entries still handling events after having been removed

- List View entries are pooled so they continue living after being removed from the list view so events like `OnUnhovered` or `OnIsSelectedChanged` can still fire after being removed from the screen
  - This change unbinds the events when the row is removed so that the widget stops handling these events and therefore avoids hitting the check in `SObjectTableRow::GetItemForThis`

diff --git a/Engine/Plugins/Runtime/CommonUI/Source/CommonUI/Public/SCommonButtonTableRow.h b/Engine/Plugins/Runtime/CommonUI/Source/CommonUI/Public/SCommonButtonTableRow.h
--- a/Engine/Plugins/Runtime/CommonUI/Source/CommonUI/Public/SCommonButtonTableRow.h
+++ b/Engine/Plugins/Runtime/CommonUI/Source/CommonUI/Public/SCommonButtonTableRow.h
@@ -46,12 +46,6 @@ public:
 			CommonButton->SetIsSelectable(SelectionMode != ESelectionMode::None);
 			CommonButton->SetIsInteractableWhenSelected(SelectionMode != ESelectionMode::None);
 
-			CommonButton->OnClicked().AddSP(this, &SCommonButtonTableRow::HandleButtonClicked);
-			CommonButton->OnDoubleClicked().AddSP(this, &SCommonButtonTableRow::HandleButtonDoubleClicked);
-			CommonButton->OnHovered().AddSP(this, &SCommonButtonTableRow::HandleButtonHovered);
-			CommonButton->OnUnhovered().AddSP(this, &SCommonButtonTableRow::HandleButtonUnhovered);
-			CommonButton->OnIsSelectedChanged().AddSP(this, &SCommonButtonTableRow::HandleButtonSelectionChanged);
-
 			CommonButton->SetTouchMethod(EButtonTouchMethod::PreciseTap);
 		}
 	}
@@ -70,6 +64,12 @@ protected:
 
 		if (UCommonButtonBase* CommonButton = Cast<UCommonButtonBase>(this->WidgetObject))
 		{
+			CommonButton->OnClicked().AddSP(this, &SCommonButtonTableRow::HandleButtonClicked);
+			CommonButton->OnDoubleClicked().AddSP(this, &SCommonButtonTableRow::HandleButtonDoubleClicked);
+			CommonButton->OnHovered().AddSP(this, &SCommonButtonTableRow::HandleButtonHovered);
+			CommonButton->OnUnhovered().AddSP(this, &SCommonButtonTableRow::HandleButtonUnhovered);
+			CommonButton->OnIsSelectedChanged().AddSP(this, &SCommonButtonTableRow::HandleButtonSelectionChanged);
+
 			if (this->IsItemSelectable())
 			{
 				const bool bIsItemSelected = this->IsItemSelected();
@@ -90,11 +90,19 @@ protected:
 	{
 		SObjectTableRow<ItemType>::ResetObjectRow();
 
-		UCommonButtonBase* CommonButton = Cast<UCommonButtonBase>(this->WidgetObject);
-		if (CommonButton && CommonButton->GetSelected())
+		if (UCommonButtonBase* CommonButton = Cast<UCommonButtonBase>(this->WidgetObject))
 		{
-			// Quietly deselect the button to reset its visual state
-			CommonButton->SetSelectedInternal(false, false, false);
+			CommonButton->OnClicked().RemoveAll(this);
+			CommonButton->OnDoubleClicked().RemoveAll(this);
+			CommonButton->OnHovered().RemoveAll(this);
+			CommonButton->OnUnhovered().RemoveAll(this);
+			CommonButton->OnIsSelectedChanged().RemoveAll(this);
+
+			if (CommonButton->GetSelected())
+			{
+				// Quietly deselect the button to reset its visual state
+				CommonButton->SetSelectedInternal(false, false, false);
+			}
 		}
 	}
 
