CommonUI: Fix not being able to use legacy inputs and Enhanced Inputs side-by-side

- Fixes the inability to bind non-EI inputs when Enhanced Input Support is enabled
- Improves handling of the various input bindings in CommonButtonBase and CommonActionWidget when both systems are used
- Also add IsValid checks when accessing CommonActionWidget's subwidgets to support overriding RebuildWidget

diff --git a/Engine/Plugins/Runtime/CommonUI/Source/CommonUI/Private/CommonActionWidget.cpp b/Engine/Plugins/Runtime/CommonUI/Source/CommonUI/Private/CommonActionWidget.cpp
--- a/Engine/Plugins/Runtime/CommonUI/Source/CommonUI/Private/CommonActionWidget.cpp
+++ b/Engine/Plugins/Runtime/CommonUI/Source/CommonUI/Private/CommonActionWidget.cpp
@@ -251,13 +251,14 @@ void UCommonActionWidget::SetEnhancedInputAction(UInputAction* InInputAction)
 {
 	UpdateBindingHandleInternal(FUIActionBindingHandle());
 	EnhancedInputAction = InInputAction;
+	InputActions.Reset();
 	UpdateActionWidget();
 }
 
 void UCommonActionWidget::SetInputAction(FDataTableRowHandle InputActionRow)
 {
 	UpdateBindingHandleInternal(FUIActionBindingHandle());
-	
+	EnhancedInputAction = nullptr;
 	InputActions.Reset();
 	InputActions.Add(InputActionRow);
 
@@ -269,16 +270,17 @@ void UCommonActionWidget::SetInputActionBinding(FUIActionBindingHandle BindingHa
 	UpdateBindingHandleInternal(BindingHandle);
 	if (TSharedPtr<FUIActionBinding> Binding = FUIActionBinding::FindBinding(BindingHandle))
 	{
-		if (CommonUI::IsEnhancedInputSupportEnabled())
+		InputActions.Reset();
+
+		const UInputAction* InputAction = Binding->InputAction.Get();
+		if (CommonUI::IsEnhancedInputSupportEnabled() && InputAction)
 		{
-			if (const UInputAction* InputAction = Binding->InputAction.Get())
-			{
-				EnhancedInputAction = const_cast<UInputAction*>(InputAction);
-			}
+			EnhancedInputAction = const_cast<UInputAction*>(InputAction);
+		}
+		else
+		{
+			InputActions.Add(Binding->LegacyActionTableRow);
 		}
-
-		InputActions.Reset();
-		InputActions.Add(Binding->LegacyActionTableRow);
 
 		UpdateActionWidget();
 	}
@@ -287,6 +289,7 @@ void UCommonActionWidget::SetInputActionBinding(FUIActionBindingHandle BindingHa
 void UCommonActionWidget::SetInputActions(TArray<FDataTableRowHandle> InInputActions)
 {
 	UpdateBindingHandleInternal(FUIActionBindingHandle());
+	EnhancedInputAction = nullptr;
 	InputActions = InInputActions;
 
 	UpdateActionWidget();
@@ -345,17 +348,23 @@ void UCommonActionWidget::UpdateActionWidget()
 						MyIcon->Invalidate(EInvalidateWidgetReason::Layout);
 					}
 
-					if (IsHeldAction())
+					if (MyProgressImage.IsValid())
 					{
-						MyProgressImage->SetVisibility(EVisibility::SelfHitTestInvisible);
+						if (IsHeldAction())
+						{
+							MyProgressImage->SetVisibility(EVisibility::SelfHitTestInvisible);
+						}
+						else
+						{
+							MyProgressImage->SetVisibility(EVisibility::Collapsed);
+						}
 					}
-					else
+
+					if (MyKeyBox.IsValid())
 					{
-						MyProgressImage->SetVisibility(EVisibility::Collapsed);
+						MyKeyBox->Invalidate(EInvalidateWidget::LayoutAndVolatility);
 					}
 
-					MyKeyBox->Invalidate(EInvalidateWidget::LayoutAndVolatility);
-
 					if (!IsDesignTime())
 					{
 						SetVisibility(ESlateVisibility::SelfHitTestInvisible);
@@ -449,7 +458,10 @@ void UCommonActionWidget::SetProgressMaterial(const FSlateBrush& InProgressMater
 		ProgressDynamicMaterial = nullptr;
 	}
 
-	MyProgressImage->SetImage(&ProgressMaterialBrush);
+	if (MyProgressImage.IsValid())
+	{
+		MyProgressImage->SetImage(&ProgressMaterialBrush);
+	}
 }
 
 void UCommonActionWidget::SetHidden(bool bAlwaysHidden)
diff --git a/Engine/Plugins/Runtime/CommonUI/Source/CommonUI/Private/CommonButtonBase.cpp b/Engine/Plugins/Runtime/CommonUI/Source/CommonUI/Private/CommonButtonBase.cpp
--- a/Engine/Plugins/Runtime/CommonUI/Source/CommonUI/Private/CommonButtonBase.cpp
+++ b/Engine/Plugins/Runtime/CommonUI/Source/CommonUI/Private/CommonButtonBase.cpp
@@ -583,7 +583,12 @@ void UCommonButtonBase::UpdateHoldData(ECommonInputType CurrentInputType)
 
 void UCommonButtonBase::BindTriggeringInputActionToClick()
 {
-	if (CommonUI::IsEnhancedInputSupportEnabled() && TriggeringEnhancedInputAction)
+	if (!TriggeredInputAction.IsNull())
+	{
+		return;
+	}
+
+	if (CommonUI::IsEnhancedInputSupportEnabled() && TriggeringEnhancedInputAction && !TriggeringBindingHandle.IsValid())
 	{
 		FBindUIActionArgs BindArgs(TriggeringEnhancedInputAction, false, FSimpleDelegate::CreateUObject(this, &UCommonButtonBase::HandleTriggeringActionCommited));
 		BindArgs.OnHoldActionProgressed.BindUObject(this, &UCommonButtonBase::NativeOnActionProgress);
@@ -594,16 +599,8 @@ void UCommonButtonBase::BindTriggeringInputActionToClick()
 		BindArgs.InputMode = InputModeOverride;
 
 		TriggeringBindingHandle = RegisterUIActionBinding(BindArgs);
-
-		return;
-	}
-
-	if (TriggeringInputAction.IsNull() || !TriggeredInputAction.IsNull())
-	{
-		return;
 	}
-
-	if (!TriggeringBindingHandle.IsValid())
+	else if (!TriggeringInputAction.IsNull() && !TriggeringBindingHandle.IsValid())
 	{
 		FBindUIActionArgs BindArgs(TriggeringInputAction, false, FSimpleDelegate::CreateUObject(this, &UCommonButtonBase::HandleTriggeringActionCommited));
 		BindArgs.OnHoldActionProgressed.BindUObject(this, &UCommonButtonBase::NativeOnActionProgress);
@@ -619,15 +616,8 @@ void UCommonButtonBase::BindTriggeringInputActionToClick()
 }
 
 void UCommonButtonBase::UnbindTriggeringInputActionToClick()
-{
-	if (CommonUI::IsEnhancedInputSupportEnabled() && TriggeringEnhancedInputAction)
-	{
-		TriggeringBindingHandle.Unregister();
-
-		return;
-	}
-
-	if (TriggeringInputAction.IsNull() || !TriggeredInputAction.IsNull())
+{	
+	if (!TriggeredInputAction.IsNull())
 	{
 		return;
 	}
@@ -1151,14 +1141,14 @@ void UCommonButtonBase::SetMinDimensions(int32 InMinWidth, int32 InMinHeight)
 
 void UCommonButtonBase::SetTriggeredInputAction(const FDataTableRowHandle &InputActionRow)
 {
-	if (ensure(TriggeringInputAction.IsNull()))
-	{
-		TriggeredInputAction = InputActionRow;
-		UpdateInputActionWidget();
+	UnbindTriggeringInputActionToClick();
 
-		OnTriggeredInputActionChanged(InputActionRow);
-	}
+	TriggeringInputAction = {};
+	TriggeringEnhancedInputAction = nullptr;
+	TriggeredInputAction = InputActionRow;
+	UpdateInputActionWidget();
 
+	OnTriggeredInputActionChanged(InputActionRow);
 }
 
 void UCommonButtonBase::SetTriggeringInputAction(const FDataTableRowHandle & InputActionRow)
@@ -1167,6 +1157,8 @@ void UCommonButtonBase::SetTriggeringInputAction(const FDataTableRowHandle & Inp
 	{
 		UnbindTriggeringInputActionToClick();
 
+		TriggeredInputAction = {};
+		TriggeringEnhancedInputAction = nullptr;
 		TriggeringInputAction = InputActionRow;
 
 		if (!IsDesignTime())
@@ -1187,6 +1179,8 @@ void UCommonButtonBase::SetTriggeringEnhancedInputAction(UInputAction* InInputAc
 	{
 		UnbindTriggeringInputActionToClick();
 
+		TriggeredInputAction = {};
+		TriggeringInputAction = {};
 		TriggeringEnhancedInputAction = InInputAction;
 
 		if (!IsDesignTime())
@@ -1260,29 +1254,16 @@ void UCommonButtonBase::UpdateInputActionWidget()
 			InputActionWidget->SetInputAction(TriggeredInputAction);
 		}
 		// Visualize the default click action when neither input action is bound and when the widget is enabled and hovered
-		else if (bShouldUseFallbackDefaultInputAction && bButtonEnabled)
+		else if (bShouldUseFallbackDefaultInputAction && bButtonEnabled && IsHovered())
 		{
-			FDataTableRowHandle HoverStateHandle;
-			UInputAction* HoverEnhancedInputAction = nullptr;
-			if (IsHovered())
-			{
-				if (bIsEnhancedInputSupportEnabled)
-				{
-					HoverEnhancedInputAction = ICommonInputModule::GetSettings().GetEnhancedInputClickAction();
-				}
-				else
-				{
-					HoverStateHandle = ICommonInputModule::GetSettings().GetDefaultClickAction();
-				}
-			}
-
-			if (bIsEnhancedInputSupportEnabled)
+			UInputAction* DefaultEnhancedClickAction = bIsEnhancedInputSupportEnabled ? ICommonInputModule::GetSettings().GetEnhancedInputClickAction() : nullptr;
+			if (DefaultEnhancedClickAction)
 			{
-				InputActionWidget->SetEnhancedInputAction(HoverEnhancedInputAction);
+				InputActionWidget->SetEnhancedInputAction(DefaultEnhancedClickAction);
 			}
 			else
 			{
-				InputActionWidget->SetInputAction(HoverStateHandle);
+				InputActionWidget->SetInputAction(ICommonInputModule::GetSettings().GetDefaultClickAction());
 			}
 		}
 		else
diff --git a/Engine/Plugins/Runtime/CommonUI/Source/CommonUI/Private/Input/UIActionRouterTypes.cpp b/Engine/Plugins/Runtime/CommonUI/Source/CommonUI/Private/Input/UIActionRouterTypes.cpp
--- a/Engine/Plugins/Runtime/CommonUI/Source/CommonUI/Private/Input/UIActionRouterTypes.cpp
+++ b/Engine/Plugins/Runtime/CommonUI/Source/CommonUI/Private/Input/UIActionRouterTypes.cpp
@@ -227,7 +227,7 @@ FUIActionBindingHandle FUIActionBinding::TryCreate(const UWidget& InBoundWidget,
 		UE_LOG(LogUIActionRouter, Error, TEXT("Cannot bind widget [%s] to action [%s] - provided legacy data table row does not resolve to valid data."), *InBoundWidget.GetName(), *BindArgs.GetActionName().ToString());
 		return FUIActionBindingHandle();
 	}
-	else if (bIsEnhancedInputSupportEnabled && !BindArgs.InputAction.IsValid())
+	else if (!BindArgs.ActionTag.IsValid() && BindArgs.LegacyActionTableRow.IsNull() && bIsEnhancedInputSupportEnabled && !BindArgs.InputAction.IsValid())
 	{
 		UE_LOG(LogUIActionRouter, Error, TEXT("Cannot bind widget [%s] to action [%s] - provided input action is invalid."), *InBoundWidget.GetName(), *BindArgs.GetActionName().ToString());
 		return FUIActionBindingHandle();
