CommonUI: Fix Common Action Widget not updating its icon when switching from showing an Enhanced Input to a Data Table input

diff --git a/Engine/Plugins/Runtime/CommonUI/Source/CommonUI/Private/CommonActionWidget.cpp b/Engine/Plugins/Runtime/CommonUI/Source/CommonUI/Private/CommonActionWidget.cpp
--- a/Engine/Plugins/Runtime/CommonUI/Source/CommonUI/Private/CommonActionWidget.cpp
+++ b/Engine/Plugins/Runtime/CommonUI/Source/CommonUI/Private/CommonActionWidget.cpp
@@ -279,6 +279,7 @@ void UCommonActionWidget::SetInputActionBinding(FUIActionBindingHandle BindingHa
 		}
 		else
 		{
+			EnhancedInputAction = nullptr;
 			InputActions.Add(Binding->LegacyActionTableRow);
 		}
 
