CommonUI: Make UCommonUIActionRouterBase::CanProcessNormalGameInput() virtual so projects can extend the logic determining when input events make it to the player controller

diff --git a/Engine/Plugins/Runtime/CommonUI/Source/CommonUI/Public/Input/CommonUIActionRouterBase.h b/Engine/Plugins/Runtime/CommonUI/Source/CommonUI/Public/Input/CommonUIActionRouterBase.h
--- a/Engine/Plugins/Runtime/CommonUI/Source/CommonUI/Public/Input/CommonUIActionRouterBase.h
+++ b/Engine/Plugins/Runtime/CommonUI/Source/CommonUI/Public/Input/CommonUIActionRouterBase.h
@@ -105,7 +105,7 @@ public:
 	UCommonInputSubsystem& GetInputSubsystem() const;
 
 	virtual ERouteUIInputResult ProcessInput(FKey Key, EInputEvent InputEvent) const;
-	bool CanProcessNormalGameInput() const;
+	virtual bool CanProcessNormalGameInput() const;
 
 	bool IsPendingTreeChange() const;
 
