PR #12380: Added a method GetControllerDataForInputType

diff --git a/Engine/Plugins/Runtime/CommonUI/Source/CommonInput/Private/CommonInputBaseTypes.cpp b/Engine/Plugins/Runtime/CommonUI/Source/CommonInput/Private/CommonInputBaseTypes.cpp
--- a/Engine/Plugins/Runtime/CommonUI/Source/CommonInput/Private/CommonInputBaseTypes.cpp
+++ b/Engine/Plugins/Runtime/CommonUI/Source/CommonInput/Private/CommonInputBaseTypes.cpp
@@ -25,6 +25,60 @@ FCommonInputKeySetBrushConfiguration::FCommonInputKeySetBrushConfiguration()
 	KeyBrush.DrawAs = ESlateBrushDrawType::Image;
 }
 
+namespace CommonUIUtils
+{
+	bool TryGetInputBrushFromDataMap(FSlateBrush& OutBrush, const FKey& InKey, const TArray<FCommonInputKeyBrushConfiguration>& InInputBrushDataMap)
+	{
+		const FCommonInputKeyBrushConfiguration* DisplayConfig = InInputBrushDataMap.FindByPredicate([&InKey](const FCommonInputKeyBrushConfiguration& KeyBrushPair) -> bool
+		{
+			return KeyBrushPair.Key == InKey;
+		});
+
+		if (DisplayConfig)
+		{
+			OutBrush = DisplayConfig->GetInputBrush();
+			return true;
+		}
+
+		return false;
+	}
+
+	bool TryGetInputBrushFromKeySets(FSlateBrush& OutBrush, const TArray<FKey>& InKeys, const TArray<FCommonInputKeySetBrushConfiguration>& InInputBrushKeySets)
+	{
+
+		const FCommonInputKeySetBrushConfiguration* DisplayConfig = InInputBrushKeySets.FindByPredicate([&InKeys](const FCommonInputKeySetBrushConfiguration& KeyBrushPair) -> bool
+		{
+			if (KeyBrushPair.Keys.Num() < 2)
+			{
+				return false;
+			}
+
+			if (InKeys.Num() == KeyBrushPair.Keys.Num())
+			{
+				for (const FKey& Key : InKeys)
+				{
+					if (!KeyBrushPair.Keys.Contains(Key))
+					{
+						return false;
+					}
+				}
+
+				return true;
+			}
+
+			return false;
+		});
+
+		if (DisplayConfig)
+		{
+			OutBrush = DisplayConfig->GetInputBrush();
+			return true;
+		}
+
+		return false;
+	}
+}
+
 bool UCommonUIInputData::NeedsLoadForServer() const
 {
 	const UUserInterfaceSettings* UISettings = GetDefault<UUserInterfaceSettings>();
@@ -39,18 +93,7 @@ bool UCommonInputBaseControllerData::NeedsLoadForServer() const
 
 bool UCommonInputBaseControllerData::TryGetInputBrush(FSlateBrush& OutBrush, const FKey& Key) const
 {
-	const FCommonInputKeyBrushConfiguration* DisplayConfig = InputBrushDataMap.FindByPredicate([&Key](const FCommonInputKeyBrushConfiguration& KeyBrushPair) -> bool
-	{
-		return KeyBrushPair.Key == Key;
-	});
-
-	if (DisplayConfig)
-	{
-		OutBrush = DisplayConfig->GetInputBrush();
-		return true;
-	}
-
-	return false;
+	return CommonUIUtils::TryGetInputBrushFromDataMap(OutBrush, Key, InputBrushDataMap);
 }
 
 bool UCommonInputBaseControllerData::TryGetInputBrush(FSlateBrush& OutBrush, const TArray<FKey>& Keys) const
@@ -62,39 +105,10 @@ bool UCommonInputBaseControllerData::TryGetInputBrush(FSlateBrush& OutBrush, con
 
 	if (Keys.Num() == 1)
 	{
-		return TryGetInputBrush(OutBrush, Keys[0]);
-	}
-
-	const FCommonInputKeySetBrushConfiguration* DisplayConfig = InputBrushKeySets.FindByPredicate([&Keys](const FCommonInputKeySetBrushConfiguration& KeyBrushPair) -> bool
-	{
-		if (KeyBrushPair.Keys.Num() < 2)
-		{
-			return false;
-		}
-
-		if (Keys.Num() == KeyBrushPair.Keys.Num())
-		{
-			for (const FKey& Key : Keys)
-			{
-				if (!KeyBrushPair.Keys.Contains(Key))
-				{
-					return false;
-				}
-			}
-
-			return true;
-		}
-
-		return false;
-	});
-
-	if (DisplayConfig)
-	{
-		OutBrush = DisplayConfig->GetInputBrush();
-		return true;
+		return CommonUIUtils::TryGetInputBrushFromDataMap(OutBrush, Keys[0], InputBrushDataMap);
 	}
 
-	return false;
+	return CommonUIUtils::TryGetInputBrushFromKeySets(OutBrush, Keys, InputBrushKeySets);
 }
 
 void UCommonInputBaseControllerData::PreSave(FObjectPreSaveContext ObjectSaveContext)
@@ -273,16 +287,9 @@ bool UCommonInputPlatformSettings::TryGetInputBrush(FSlateBrush& OutBrush, FKey
 {
 	InitializeControllerData();
 
-	for (const TSubclassOf<UCommonInputBaseControllerData>& ControllerDataPtr : ControllerDataClasses)
+	if (const UCommonInputBaseControllerData* DefaultControllerData = GetControllerDataForInputType(InputType, GamepadName))
 	{
-		const UCommonInputBaseControllerData* DefaultControllerData = ControllerDataPtr.GetDefaultObject();
-		if (DefaultControllerData && DefaultControllerData->InputType == InputType)
-		{
-			if (DefaultControllerData->InputType != ECommonInputType::Gamepad || DefaultControllerData->GamepadName == GamepadName)
-			{
-				return DefaultControllerData->TryGetInputBrush(OutBrush, Key);
-			}
-		}
+		return DefaultControllerData->TryGetInputBrush(OutBrush, Key);
 	}
 
 	return false;
@@ -292,6 +299,18 @@ bool UCommonInputPlatformSettings::TryGetInputBrush(FSlateBrush& OutBrush, const
 {
 	InitializeControllerData();
 
+	if (const UCommonInputBaseControllerData* DefaultControllerData = GetControllerDataForInputType(InputType, GamepadName))
+	{
+		return DefaultControllerData->TryGetInputBrush(OutBrush, Keys);
+	}
+
+	return false;
+}
+
+const UCommonInputBaseControllerData* UCommonInputPlatformSettings::GetControllerDataForInputType(ECommonInputType InputType, const FName GamepadName) const
+{
+	InitializeControllerData();
+
 	for (const TSubclassOf<UCommonInputBaseControllerData>& ControllerDataPtr : ControllerDataClasses)
 	{
 		const UCommonInputBaseControllerData* DefaultControllerData = ControllerDataPtr.GetDefaultObject();
@@ -299,12 +318,12 @@ bool UCommonInputPlatformSettings::TryGetInputBrush(FSlateBrush& OutBrush, const
 		{
 			if (DefaultControllerData->InputType != ECommonInputType::Gamepad || DefaultControllerData->GamepadName == GamepadName)
 			{
-				return DefaultControllerData->TryGetInputBrush(OutBrush, Keys);
+				return DefaultControllerData;
 			}
 		}
 	}
 
-	return false;
+	return nullptr;
 }
 
 FName UCommonInputPlatformSettings::GetBestGamepadNameForHardware(FName CurrentGamepadName, FName InputDeviceName, const FString& HardwareDeviceIdentifier)
diff --git a/Engine/Plugins/Runtime/CommonUI/Source/CommonInput/Public/CommonInputBaseTypes.h b/Engine/Plugins/Runtime/CommonUI/Source/CommonInput/Public/CommonInputBaseTypes.h
--- a/Engine/Plugins/Runtime/CommonUI/Source/CommonInput/Public/CommonInputBaseTypes.h
+++ b/Engine/Plugins/Runtime/CommonUI/Source/CommonInput/Public/CommonInputBaseTypes.h
@@ -90,6 +90,12 @@ struct FInputHoldData
 	float HoldRollbackTime = 0.0;
 };
 
+namespace CommonUIUtils
+{
+	COMMONINPUT_API bool TryGetInputBrushFromDataMap(FSlateBrush& OutBrush, const FKey& InKey, const TArray<FCommonInputKeyBrushConfiguration>& InInputBrushDataMap);
+	COMMONINPUT_API bool TryGetInputBrushFromKeySets(FSlateBrush& OutBrush, const TArray<FKey>& InKeys, const TArray<FCommonInputKeySetBrushConfiguration>& InInputBrushKeySets);
+}
+
 /* Derive from this class to store the Input data. It is referenced in the Common Input Settings, found in the project settings UI. */
 UCLASS(Abstract, Blueprintable, ClassGroup = Input, meta = (Category = "Common Input"))
 class COMMONINPUT_API UCommonUIInputData : public UObject
@@ -224,6 +230,8 @@ public:
 	bool TryGetInputBrush(FSlateBrush& OutBrush, FKey Key, ECommonInputType InputType, const FName GamepadName) const;
 	bool TryGetInputBrush(FSlateBrush& OutBrush, const TArray<FKey>& Keys, ECommonInputType InputType, const FName GamepadName) const;
 
+	const UCommonInputBaseControllerData* GetControllerDataForInputType(ECommonInputType InputType, const FName GamepadName) const;
+
 	FName GetBestGamepadNameForHardware(FName CurrentGamepadName, FName InputDeviceName, const FString& HardwareDeviceIdentifier); 
 
 	ECommonInputType GetDefaultInputType() const
