Input  Router: Remove temporary Cvar and leave it enabled to true permanently. Updated the comments to reflect the new true cvar functionality and simplified code where possible.

diff --git a/Engine/Plugins/Runtime/CommonUI/Source/CommonUI/Private/Input/CommonUIActionRouterBase.cpp b/Engine/Plugins/Runtime/CommonUI/Source/CommonUI/Private/Input/CommonUIActionRouterBase.cpp
--- a/Engine/Plugins/Runtime/CommonUI/Source/CommonUI/Private/Input/CommonUIActionRouterBase.cpp
+++ b/Engine/Plugins/Runtime/CommonUI/Source/CommonUI/Private/Input/CommonUIActionRouterBase.cpp
@@ -79,11 +79,6 @@ static FAutoConsoleVariableRef CVarCheckGameViewportClientValid(
 	bCheckGameViewportClientValid,
 	TEXT("Log error when CommonUI is used without the current game viewport deriving from CommonGameViewportClient."));
 
-TAutoConsoleVariable<bool> CvarEarlyOutRefreshActionDomainLeafNodeConfig(
-	TEXT("CommonUI.Debug.EarlyOutRefreshActionDomainLeafNodeConfig"),
-	true,
-	TEXT("When true early out in RefreshActionDomainLeafNodeConfig if there is an active root node."));
-
 //////////////////////////////////////////////////////////////////////////
 
 FGlobalUITags FGlobalUITags::GUITags;
@@ -822,44 +817,16 @@ void UCommonUIActionRouterBase::HandleRootNodeDeactivated(TWeakPtr<FActivatableT
 		SetActiveRoot(nullptr);
 	}
 
-	
-	if (!CvarEarlyOutRefreshActionDomainLeafNodeConfig->GetBool())
+	// In the case that the activatable tree was enabled we need to re-establish input for the highest paint layer node in action domain nodes.
+	if (bIsActivatableTreeEnabled)
 	{
-		bool bActivatedRootNodeExists = false;
-		for (const FActivatableTreeRootRef& Root : RootNodes)
+		if (bWarnAllWidgetsDeactivated)
 		{
-			if (Root->IsWidgetActivated())
-			{
-				bActivatedRootNodeExists = true;
-				break;
-			}
+			UE_LOG(LogUIActionRouter, Warning, TEXT("All widgets deactivated. Existing input config set: %s"),
+				ActiveInputConfig.IsSet() ? TEXT("Yes - the current input config is lingering from a deactivated widget.") : TEXT("No."));
 		}
 
-		// In the case that all root nodes are not activated we need to re-establish input for the highest paint layer node in action domain nodes.
-		if (!bActivatedRootNodeExists && bIsActivatableTreeEnabled)
-		{
-			if (bWarnAllWidgetsDeactivated)
-			{
-				UE_LOG(LogUIActionRouter, Warning, TEXT("All widgets deactivated. Existing input config set: %s"),
-					   ActiveInputConfig.IsSet() ? TEXT("Yes - the current input config is lingering from a deactivated widget.") : TEXT("No."));
-			}
-
-			RefreshActionDomainLeafNodeConfig();
-		}
-	}
-	else
-	{
-		// In the case that all root nodes are not activated we need to re-establish input for the highest paint layer node in action domain nodes.
-		if (bIsActivatableTreeEnabled)
-		{
-			if (bWarnAllWidgetsDeactivated)
-			{
-				UE_LOG(LogUIActionRouter, Warning, TEXT("All widgets deactivated. Existing input config set: %s"),
-					   ActiveInputConfig.IsSet() ? TEXT("Yes - the current input config is lingering from a deactivated widget.") : TEXT("No."));
-			}
-
-			RefreshActionDomainLeafNodeConfig();
-		}
+		RefreshActionDomainLeafNodeConfig();
 	}
 }
 
@@ -1520,27 +1487,11 @@ void UCommonUIActionRouterBase::SetActiveUIInputConfig(const FUIInputConfig& New
 
 void UCommonUIActionRouterBase::RefreshActionDomainLeafNodeConfig()
 {
-	if (CvarEarlyOutRefreshActionDomainLeafNodeConfig->GetBool())
-	{
-		// We don't want to refresh if the activatable tree is not enabled as we don't want input mode changes when dormant
-		bool bActivatedRootNodeExists = false;
-		for (const FActivatableTreeRootRef& Root : RootNodes)
-		{
-			if (Root->IsWidgetActivated())
-			{
-				bActivatedRootNodeExists = true;
-				break;
-			}
-		}
-
-		if (bActivatedRootNodeExists)
-		{
-			return;
-		}
-	}
-	else
+	// We don't want to refresh if there is an activated root node as we don't want input mode changes
+	bool bActivatedRootNodeExists = false;
+	for (const FActivatableTreeRootRef& Root : RootNodes)
 	{
-		if (!bIsActivatableTreeEnabled)
+		if (Root->IsWidgetActivated())
 		{
 			return;
 		}
