Common UI: add enhanced input getter to CommonActionWidget

diff --git a/Engine/Plugins/Runtime/CommonUI/Source/CommonUI/Private/CommonActionWidget.cpp b/Engine/Plugins/Runtime/CommonUI/Source/CommonUI/Private/CommonActionWidget.cpp
--- a/Engine/Plugins/Runtime/CommonUI/Source/CommonUI/Private/CommonActionWidget.cpp
+++ b/Engine/Plugins/Runtime/CommonUI/Source/CommonUI/Private/CommonActionWidget.cpp
@@ -255,6 +255,11 @@ void UCommonActionWidget::SetEnhancedInputAction(UInputAction* InInputAction)
 	UpdateActionWidget();
 }
 
+const UInputAction* UCommonActionWidget::GetEnhancedInputAction() const
+{
+	return EnhancedInputAction;
+}
+
 void UCommonActionWidget::SetInputAction(FDataTableRowHandle InputActionRow)
 {
 	UpdateBindingHandleInternal(FUIActionBindingHandle());
diff --git a/Engine/Plugins/Runtime/CommonUI/Source/CommonUI/Public/CommonActionWidget.h b/Engine/Plugins/Runtime/CommonUI/Source/CommonUI/Public/CommonActionWidget.h
--- a/Engine/Plugins/Runtime/CommonUI/Source/CommonUI/Public/CommonActionWidget.h
+++ b/Engine/Plugins/Runtime/CommonUI/Source/CommonUI/Public/CommonActionWidget.h
@@ -48,6 +48,9 @@ public:
 	UFUNCTION(BlueprintCallable, Category = CommonActionWidget)
 	void SetEnhancedInputAction(UInputAction* InInputAction);
 
+	UFUNCTION(BlueprintCallable, Category = CommonActionWidget)
+	const UInputAction* GetEnhancedInputAction() const;
+
 	UFUNCTION(BlueprintCallable, Category = CommonActionWidget)
 	void SetInputAction(FDataTableRowHandle InputActionRow);
 
@@ -106,7 +109,7 @@ protected:
 	/**
 	 * Input Action this common action widget is intended to represent. Optional if using EnhancedInputs
 	 */
-	UPROPERTY(EditAnywhere, BlueprintReadOnly, Category = CommonActionWidget, meta = (EditCondition = "CommonInput.CommonInputSettings.IsEnhancedInputSupportEnabled", EditConditionHides))
+	UPROPERTY(EditAnywhere, BlueprintReadOnly, Getter, Category = CommonActionWidget, meta = (EditCondition = "CommonInput.CommonInputSettings.IsEnhancedInputSupportEnabled", EditConditionHides))
 	TObjectPtr<class UInputAction> EnhancedInputAction;
 
 	FUIActionBindingHandle DisplayedBindingHandle;
