PR #12416: Minor editor fix for Common UI uncategorized widgets in widget blueprint palette

diff --git a/Engine/Plugins/Runtime/CommonUI/Source/CommonUI/Private/CommonUserWidget.cpp b/Engine/Plugins/Runtime/CommonUI/Source/CommonUI/Private/CommonUserWidget.cpp
--- a/Engine/Plugins/Runtime/CommonUI/Source/CommonUI/Private/CommonUserWidget.cpp
+++ b/Engine/Plugins/Runtime/CommonUI/Source/CommonUI/Private/CommonUserWidget.cpp
@@ -5,6 +5,7 @@
 #include "Engine/GameInstance.h"
 #include "CommonInputSubsystem.h"
 #include "CommonUISubsystemBase.h"
+#include "CommonWidgetPaletteCategories.h"
 #include "Input/CommonUIActionRouterBase.h"
 #include "Input/CommonUIInputTypes.h"
 
@@ -206,4 +207,9 @@ void UCommonUserWidget::NativeDestruct()
 	Super::NativeDestruct();
 }
 
-
+#if WITH_EDITOR
+const FText UCommonUserWidget::GetPaletteCategory()
+{
+	return CommonWidgetPaletteCategories::Default;
+}
+#endif // WITH_EDITOR
diff --git a/Engine/Plugins/Runtime/CommonUI/Source/CommonUI/Private/CommonVideoPlayer.cpp b/Engine/Plugins/Runtime/CommonUI/Source/CommonUI/Private/CommonVideoPlayer.cpp
--- a/Engine/Plugins/Runtime/CommonUI/Source/CommonUI/Private/CommonVideoPlayer.cpp
+++ b/Engine/Plugins/Runtime/CommonUI/Source/CommonUI/Private/CommonVideoPlayer.cpp
@@ -11,6 +11,7 @@
 #include "MediaSoundComponent.h"
 #include "IMediaEventSink.h"
 #include "Widgets/Images/SImage.h"
+#include "CommonWidgetPaletteCategories.h"
 
 #include UE_INLINE_GENERATED_CPP_BY_NAME(CommonVideoPlayer)
 
@@ -241,6 +242,13 @@ void UCommonVideoPlayer::PlaybackTick(double InCurrentTime, float InDeltaTime)
 	}
 }
 
+#if WITH_EDITOR
+const FText UCommonVideoPlayer::GetPaletteCategory()
+{
+	return CommonWidgetPaletteCategories::Default;
+}
+#endif // WITH_EDITOR
+
 EActiveTimerReturnType UCommonVideoPlayer::HandlePlaybackTick(double InCurrentTime, float InDeltaTime)
 {
 	PlaybackTick(InCurrentTime, InDeltaTime);
diff --git a/Engine/Plugins/Runtime/CommonUI/Source/CommonUI/Private/Widgets/CommonActivatableWidgetContainer.cpp b/Engine/Plugins/Runtime/CommonUI/Source/CommonUI/Private/Widgets/CommonActivatableWidgetContainer.cpp
--- a/Engine/Plugins/Runtime/CommonUI/Source/CommonUI/Private/Widgets/CommonActivatableWidgetContainer.cpp
+++ b/Engine/Plugins/Runtime/CommonUI/Source/CommonUI/Private/Widgets/CommonActivatableWidgetContainer.cpp
@@ -3,6 +3,7 @@
 #include "Widgets/CommonActivatableWidgetContainer.h"
 #include "CommonActivatableWidget.h"
 #include "CommonUIPrivate.h"
+#include "CommonWidgetPaletteCategories.h"
 #include "Slate/SCommonAnimatedSwitcher.h"
 #include "Widgets/SOverlay.h"
 #include "Widgets/Layout/SSpacer.h"
@@ -168,6 +169,13 @@ void UCommonActivatableWidgetContainerBase::SetSwitcherIndex(int32 TargetIndex,
 	}
 }
 
+#if WITH_EDITOR
+const FText UCommonActivatableWidgetContainerBase::GetPaletteCategory()
+{
+	return CommonWidgetPaletteCategories::Default;
+}
+#endif // WITH_EDITOR
+
 UCommonActivatableWidget* UCommonActivatableWidgetContainerBase::BP_AddWidget(TSubclassOf<UCommonActivatableWidget> ActivatableWidgetClass)
 {
 	return AddWidgetInternal(ActivatableWidgetClass, [](UCommonActivatableWidget&) {});
diff --git a/Engine/Plugins/Runtime/CommonUI/Source/CommonUI/Public/CommonUserWidget.h b/Engine/Plugins/Runtime/CommonUI/Source/CommonUI/Public/CommonUserWidget.h
--- a/Engine/Plugins/Runtime/CommonUI/Source/CommonUI/Public/CommonUserWidget.h
+++ b/Engine/Plugins/Runtime/CommonUI/Source/CommonUI/Public/CommonUserWidget.h
@@ -83,6 +83,10 @@ protected:
 	void RegisterScrollRecipient(const UWidget& AnalogScrollRecipient);
 	void UnregisterScrollRecipient(const UWidget& AnalogScrollRecipient);
 
+#if WITH_EDITOR
+	virtual const FText GetPaletteCategory() override;
+#endif // WITH_EDITOR
+
 	/** True to generally display this widget's actions in the action bar, assuming it has actions. */
 	UPROPERTY(EditAnywhere, BlueprintReadOnly, Category = Input, meta = (AllowPrivateAccess = true))
 	bool bDisplayInActionBar = false;
diff --git a/Engine/Plugins/Runtime/CommonUI/Source/CommonUI/Public/CommonVideoPlayer.h b/Engine/Plugins/Runtime/CommonUI/Source/CommonUI/Public/CommonVideoPlayer.h
--- a/Engine/Plugins/Runtime/CommonUI/Source/CommonUI/Public/CommonVideoPlayer.h
+++ b/Engine/Plugins/Runtime/CommonUI/Source/CommonUI/Public/CommonVideoPlayer.h
@@ -80,6 +80,10 @@ protected:
 	virtual void HandleMediaPlayerEvent(EMediaEvent EventType);
 	virtual void PlaybackTick(double InCurrentTime, float InDeltaTime);
 
+#if WITH_EDITOR
+	virtual const FText GetPaletteCategory() override;
+#endif // WITH_EDITOR
+
 private:
 	EActiveTimerReturnType HandlePlaybackTick(double InCurrentTime, float InDeltaTime);
 
diff --git a/Engine/Plugins/Runtime/CommonUI/Source/CommonUI/Public/Widgets/CommonActivatableWidgetContainer.h b/Engine/Plugins/Runtime/CommonUI/Source/CommonUI/Public/Widgets/CommonActivatableWidgetContainer.h
--- a/Engine/Plugins/Runtime/CommonUI/Source/CommonUI/Public/Widgets/CommonActivatableWidgetContainer.h
+++ b/Engine/Plugins/Runtime/CommonUI/Source/CommonUI/Public/Widgets/CommonActivatableWidgetContainer.h
@@ -97,6 +97,10 @@ protected:
 
 	void SetSwitcherIndex(int32 TargetIndex, bool bInstantTransition = false);
 
+#if WITH_EDITOR
+	virtual const FText GetPaletteCategory() override;
+#endif // WITH_EDITOR
+
 	/** The type of transition to play between widgets */
 	UPROPERTY(EditAnywhere, Category = Transition)
 	ECommonSwitcherTransition TransitionType;
