PR #12059: CommonUI adding TransitionType setter

diff --git a/Engine/Plugins/Runtime/CommonUI/Source/CommonUI/Private/Slate/SCommonAnimatedSwitcher.cpp b/Engine/Plugins/Runtime/CommonUI/Source/CommonUI/Private/Slate/SCommonAnimatedSwitcher.cpp
--- a/Engine/Plugins/Runtime/CommonUI/Source/CommonUI/Private/Slate/SCommonAnimatedSwitcher.cpp
+++ b/Engine/Plugins/Runtime/CommonUI/Source/CommonUI/Private/Slate/SCommonAnimatedSwitcher.cpp
@@ -141,6 +141,17 @@ void SCommonAnimatedSwitcher::SetTransition(float Duration, ETransitionCurve Cur
 	TransitionSequence = FCurveSequence(0.f, Duration * 0.5f, TransitionCurveToCurveEaseFunction(Curve));
 }
 
+void SCommonAnimatedSwitcher::SetTransition(float Duration, ETransitionCurve Curve, ECommonSwitcherTransition NewTransitionType)
+{
+	SetTransition(Duration, Curve);
+	SetTransitionType(NewTransitionType);
+}
+
+void SCommonAnimatedSwitcher::SetTransitionType(ECommonSwitcherTransition NewTransitionType)
+{
+	TransitionType = NewTransitionType;
+}
+
 void SCommonAnimatedSwitcher::OnSlotAdded(int32 AddedIndex)
 {
 	SWidgetSwitcher::OnSlotAdded(AddedIndex);
diff --git a/Engine/Plugins/Runtime/CommonUI/Source/CommonUI/Public/Slate/SCommonAnimatedSwitcher.h b/Engine/Plugins/Runtime/CommonUI/Source/CommonUI/Public/Slate/SCommonAnimatedSwitcher.h
--- a/Engine/Plugins/Runtime/CommonUI/Source/CommonUI/Public/Slate/SCommonAnimatedSwitcher.h
+++ b/Engine/Plugins/Runtime/CommonUI/Source/CommonUI/Public/Slate/SCommonAnimatedSwitcher.h
@@ -106,6 +106,8 @@ public:
 	FSlot* GetChildSlot(int32 SlotIndex);
 
 	void SetTransition(float Duration, ETransitionCurve Curve);
+	void SetTransition(float Duration, ETransitionCurve Curve, ECommonSwitcherTransition NewTransitionType);
+	void SetTransitionType(ECommonSwitcherTransition NewTransitionType);
 
 	bool IsTransitionPlaying() const;
 	
