PR #12107: CommonUI - Back Action can now have Display Name overridden in the Activatable Widgets

diff --git a/Engine/Plugins/Runtime/CommonUI/Source/CommonUI/Private/CommonActivatableWidget.cpp b/Engine/Plugins/Runtime/CommonUI/Source/CommonUI/Private/CommonActivatableWidget.cpp
--- a/Engine/Plugins/Runtime/CommonUI/Source/CommonUI/Private/CommonActivatableWidget.cpp
+++ b/Engine/Plugins/Runtime/CommonUI/Source/CommonUI/Private/CommonActivatableWidget.cpp
@@ -37,6 +37,7 @@ void UCommonActivatableWidget::NativeConstruct()
 		{
 			FBindUIActionArgs BindArgs(ICommonInputModule::GetSettings().GetEnhancedInputBackAction(), FSimpleDelegate::CreateUObject(this, &UCommonActivatableWidget::HandleBackAction));
 			BindArgs.bDisplayInActionBar = bIsBackActionDisplayedInActionBar;
+			BindArgs.OverrideDisplayName = OverrideBackActionDisplayName;
 
 			DefaultBackActionHandle = RegisterUIActionBinding(BindArgs);
 		}
@@ -44,6 +45,7 @@ void UCommonActivatableWidget::NativeConstruct()
 		{
 			FBindUIActionArgs BindArgs(ICommonInputModule::GetSettings().GetDefaultBackAction(), FSimpleDelegate::CreateUObject(this, &UCommonActivatableWidget::HandleBackAction));
 			BindArgs.bDisplayInActionBar = bIsBackActionDisplayedInActionBar;
+			BindArgs.OverrideDisplayName = OverrideBackActionDisplayName;
 
 			DefaultBackActionHandle = RegisterUIActionBinding(BindArgs);
 		}
diff --git a/Engine/Plugins/Runtime/CommonUI/Source/CommonUI/Public/CommonActivatableWidget.h b/Engine/Plugins/Runtime/CommonUI/Source/CommonUI/Public/CommonActivatableWidget.h
--- a/Engine/Plugins/Runtime/CommonUI/Source/CommonUI/Public/CommonActivatableWidget.h
+++ b/Engine/Plugins/Runtime/CommonUI/Source/CommonUI/Public/CommonActivatableWidget.h
@@ -174,6 +174,10 @@ protected:
 	UPROPERTY(EditAnywhere, Category = Back)
 	bool bIsBackActionDisplayedInActionBar = false;
 
+	/** You can specify an override for display name for the back action. If left empty, the default will be used. */
+	UPROPERTY(EditAnywhere, Category = Back)
+	FText OverrideBackActionDisplayName;
+
 	/** True to automatically activate upon construction */
 	UPROPERTY(EditAnywhere, Category = Activation)
 	bool bAutoActivate = false;
