CommonUI: Fix FBindUIActionArgs::OverrideDisplayName / FUIActionBinding::ActionDisplayName not working for enhanced input actions

diff --git a/Engine/Plugins/Runtime/CommonUI/Source/CommonUI/Private/Input/UIActionRouterTypes.cpp b/Engine/Plugins/Runtime/CommonUI/Source/CommonUI/Private/Input/UIActionRouterTypes.cpp
--- a/Engine/Plugins/Runtime/CommonUI/Source/CommonUI/Private/Input/UIActionRouterTypes.cpp
+++ b/Engine/Plugins/Runtime/CommonUI/Source/CommonUI/Private/Input/UIActionRouterTypes.cpp
@@ -155,6 +155,11 @@ FUIActionBinding::FUIActionBinding(const UWidget& InBoundWidget, const FBindUIAc
 	}
 	else if (CommonUI::IsEnhancedInputSupportEnabled() && BindArgs.InputAction.IsValid())
 	{
+		if (ActionDisplayName.IsEmpty())
+		{
+			ActionDisplayName = BindArgs.InputAction->ActionDescription;
+		}
+
 		// Nothing else to do if we have an enhanced input action,
 		// the input action itself will be queried against for keys later
 	}
@@ -486,14 +491,6 @@ FText FUIActionBindingHandle::GetDisplayName() const
 	{
 		if (const UCommonInputSubsystem* CommonInputSubsystem = Binding->BoundWidget.IsValid() ? UCommonInputSubsystem::Get(Binding->BoundWidget->GetOwningLocalPlayer()) : nullptr)
 		{
-			if (CommonUI::IsEnhancedInputSupportEnabled())
-			{
-				if (const TObjectPtr<const UInputAction> InputAction = Binding->InputAction.Get())
-				{
-					return InputAction->ActionDescription;
-				}
-			}
-
 			const ECommonInputType CurrentInputType = CommonInputSubsystem->GetCurrentInputType();
 
 			for (const FUIActionKeyMapping& HoldMapping : Binding->HoldMappings)
