CommonUI: Sort the action domain root node list before fetching it to ensure it's always accurate

- Limit sorting-on-fetch to once per frame
- Fixes issues where the top-most widget would not become the active node until another widget was added to the list

diff --git a/Engine/Plugins/Runtime/CommonUI/Source/CommonUI/Private/Input/CommonUIActionRouterBase.cpp b/Engine/Plugins/Runtime/CommonUI/Source/CommonUI/Private/Input/CommonUIActionRouterBase.cpp
--- a/Engine/Plugins/Runtime/CommonUI/Source/CommonUI/Private/Input/CommonUIActionRouterBase.cpp
+++ b/Engine/Plugins/Runtime/CommonUI/Source/CommonUI/Private/Input/CommonUIActionRouterBase.cpp
@@ -390,7 +390,7 @@ TArray<const UWidget*> UCommonUIActionRouterBase::GatherActiveAnalogScrollRecipi
 			{
 				if (const FActionDomainSortedRootList* SortedRootList = ActionDomainRootNodes.Find(ActionDomain))
 				{
-					for (const FActivatableTreeRootRef& RootNode : SortedRootList->RootList)
+					for (const FActivatableTreeRootRef& RootNode : SortedRootList->GetRootList())
 					{
 						// only return the first of the root nodes that's in the sorted root list in action domain order.
 						if (RootNode->IsReceivingInput() && RootNode->DoesWidgetSupportActivationFocus())
@@ -422,7 +422,7 @@ TArray<FUIActionBindingHandle> UCommonUIActionRouterBase::GatherActiveBindings()
 		{
 			if (const FActionDomainSortedRootList* SortedRootList = ActionDomainRootNodes.Find(ActionDomain))
 			{
-				for (const FActivatableTreeRootRef& RootNode : SortedRootList->RootList)
+				for (const FActivatableTreeRootRef& RootNode : SortedRootList->GetRootList())
 				{
 					if (RootNode->IsReceivingInput() && RootNode->IsWidgetActivated())
 					{
@@ -784,7 +784,7 @@ void UCommonUIActionRouterBase::HandleRootNodeActivated(TWeakPtr<FActivatableTre
 				{
 					if (FActionDomainSortedRootList* SortedRootList = ActionDomainRootNodes.Find(ActionDomain))
 					{
-						for (FActivatableTreeRootRef& RootNode : SortedRootList->RootList)
+						for (FActivatableTreeRootRef& RootNode : SortedRootList->GetRootList())
 						{
 							if (RootNode->IsReceivingInput())
 							{
@@ -858,7 +858,7 @@ void UCommonUIActionRouterBase::HandlePostGarbageCollect()
 
 	for (TPair<TObjectPtr<UCommonInputActionDomain>, FActionDomainSortedRootList>& Pair : ActionDomainRootNodes)
 	{
-		for (TArray<FActivatableTreeRootRef>::TIterator Iter = Pair.Value.RootList.CreateIterator(); Iter; ++Iter)
+		for (TArray<FActivatableTreeRootRef>::TIterator Iter = Pair.Value.GetRootList().CreateIterator(); Iter; ++Iter)
 		{
 			if (!Iter->Get().IsWidgetValid())
 			{
@@ -896,7 +896,7 @@ bool UCommonUIActionRouterBase::ProcessInputOnActionDomains(ECommonInputMode Act
 		bool bInputEventHandledInDomain = false;
 		bool bDomainHadActiveRoots = false;
 
-		for (const FActivatableTreeRootRef& RootNode : SortedRootList->RootList)
+		for (const FActivatableTreeRootRef& RootNode : SortedRootList->GetRootList())
 		{
 			if (RootNode->IsWidgetActivated())
 			{
@@ -944,7 +944,7 @@ EProcessHoldActionResult UCommonUIActionRouterBase::ProcessHoldInputOnActionDoma
 		bool bInputEventHandledInDomain = false;
 		bool bDomainHadActiveRoots = false;
 
-		for (const FActivatableTreeRootRef& RootNode : SortedRootList->RootList)
+		for (const FActivatableTreeRootRef& RootNode : SortedRootList->GetRootList())
 		{
 			if (RootNode->IsReceivingInput() && HoldActionResult == EProcessHoldActionResult::Unhandled)
 			{
@@ -1370,7 +1370,7 @@ FActivatableTreeNodePtr UCommonUIActionRouterBase::FindNode(const UCommonActivat
 				
 		for (const TPair<TObjectPtr<UCommonInputActionDomain>, FActionDomainSortedRootList>& Pair : ActionDomainRootNodes)
 		{
-			for (const FActivatableTreeRootRef& RootNode : Pair.Value.RootList)
+			for (const FActivatableTreeRootRef& RootNode : Pair.Value.GetRootList())
 			{
 				FActivatableTreeNodePtr FoundNode = FindNodeRecursive(RootNode, *Widget);
 				if (FoundNode.IsValid())
@@ -1503,7 +1503,7 @@ void UCommonUIActionRouterBase::RefreshActionDomainLeafNodeConfig()
 		{
 			if (FActionDomainSortedRootList* SortedRootList = ActionDomainRootNodes.Find(ActionDomain))
 			{
-				for (FActivatableTreeRootRef& RootNode : SortedRootList->RootList)
+				for (FActivatableTreeRootRef& RootNode : SortedRootList->GetRootList())
 				{
 					// only root nodes that are actively receiving input and supports widget activation focus
 					// should update leaf nodes and have input config applied. This will also update focus.
@@ -1806,13 +1806,8 @@ static const FAutoConsoleCommandWithWorld DumpInputConfigCommand(
 
 void UCommonUIActionRouterBase::FActionDomainSortedRootList::Add(FActivatableTreeRootRef RootNode)
 {
-	auto SortFunc = [](const FActivatableTreeRootRef& A, const FActivatableTreeRootRef& B) 
-		{
-			return A->GetLastPaintLayer() > B->GetLastPaintLayer();
-		};
-
 	RootList.Add(RootNode);
-	RootList.Sort(SortFunc);
+	Sort();
 }
 
 int32 UCommonUIActionRouterBase::FActionDomainSortedRootList::Remove(FActivatableTreeRootRef RootNode)
@@ -1825,3 +1820,14 @@ bool UCommonUIActionRouterBase::FActionDomainSortedRootList::Contains(FActivatab
 	return RootList.Contains(RootNode);
 }
 
+void UCommonUIActionRouterBase::FActionDomainSortedRootList::Sort()
+{
+	auto SortFunc = [](const FActivatableTreeRootRef& A, const FActivatableTreeRootRef& B) 
+	{
+		return A->GetLastPaintLayer() > B->GetLastPaintLayer();
+	};
+
+	RootList.Sort(SortFunc);
+	bDidSortThisFrame = true;
+}
+
diff --git a/Engine/Plugins/Runtime/CommonUI/Source/CommonUI/Public/Input/CommonUIActionRouterBase.h b/Engine/Plugins/Runtime/CommonUI/Source/CommonUI/Public/Input/CommonUIActionRouterBase.h
--- a/Engine/Plugins/Runtime/CommonUI/Source/CommonUI/Public/Input/CommonUIActionRouterBase.h
+++ b/Engine/Plugins/Runtime/CommonUI/Source/CommonUI/Public/Input/CommonUIActionRouterBase.h
@@ -12,6 +12,7 @@
 #include "InputCoreTypes.h"
 #include "UObject/WeakObjectPtr.h"
 #include "Framework/Application/SlateApplication.h"
+#include "Misc/FrameValue.h"
 #include "CommonUIActionRouterBase.generated.h"
 
 class SWidget;
@@ -247,11 +248,28 @@ PRAGMA_ENABLE_DEPRECATION_WARNINGS
 
 	mutable TArray<FKey> HeldKeys;
 
-	/** A wrapper around TArray that keeps RootList sorted by PaintLayer during insertion */
+	/** A wrapper around TArray that keeps RootList sorted by PaintLayer */
 	struct FActionDomainSortedRootList
 	{
-		TArray<FActivatableTreeRootRef> RootList;
-
+		TArray<FActivatableTreeRootRef>& GetRootList()
+		{
+			if (!bDidSortThisFrame.TryGetValue(false))
+			{
+				// Make sure the list sorting is always up to date
+				Sort();
+			}
+			return RootList;
+		}
+		const TArray<FActivatableTreeRootRef>& GetRootList() const
+		{
+			if (!bDidSortThisFrame.TryGetValue(false))
+			{
+				// Make sure the list sorting is always up to date
+				const_cast<FActionDomainSortedRootList*>(this)->Sort();
+			}
+			return RootList;
+		}
+		
 		// Inserts RootNode into RootList based on its paint layer
 		void Add(FActivatableTreeRootRef RootNode);
 
@@ -260,6 +278,12 @@ PRAGMA_ENABLE_DEPRECATION_WARNINGS
 
 		// Trivial Contains check
 		bool Contains(FActivatableTreeRootRef RootNode) const;
+		
+		void Sort();
+
+	private:
+		TArray<FActivatableTreeRootRef> RootList;
+		TFrameValue<bool> bDidSortThisFrame;
 	};
 
 	TMap<TObjectPtr<UCommonInputActionDomain>, FActionDomainSortedRootList> ActionDomainRootNodes;
