CommonUI: Fix action domain widgets overriding focus and input config when they're not the active widget

- Added FindActiveActionDomainRootNode() to be the source of truth for which action domain widget is the active one

diff --git a/Engine/Plugins/Runtime/CommonUI/Source/CommonUI/Private/Input/CommonUIActionRouterBase.cpp b/Engine/Plugins/Runtime/CommonUI/Source/CommonUI/Private/Input/CommonUIActionRouterBase.cpp
--- a/Engine/Plugins/Runtime/CommonUI/Source/CommonUI/Private/Input/CommonUIActionRouterBase.cpp
+++ b/Engine/Plugins/Runtime/CommonUI/Source/CommonUI/Private/Input/CommonUIActionRouterBase.cpp
@@ -382,25 +382,9 @@ TArray<const UWidget*> UCommonUIActionRouterBase::GatherActiveAnalogScrollRecipi
 	{
 		return ActiveRootNode->GatherScrollRecipients();
 	}
-	else
+	else if (FActivatableTreeRootPtr RootNode = FindActiveActionDomainRootNode())
 	{
-		if (const UCommonInputActionDomainTable* ActionDomainTable = GetActionDomainTable())
-		{
-			for (const UCommonInputActionDomain* ActionDomain : ActionDomainTable->ActionDomains)
-			{
-				if (const FActionDomainSortedRootList* SortedRootList = ActionDomainRootNodes.Find(ActionDomain))
-				{
-					for (const FActivatableTreeRootRef& RootNode : SortedRootList->GetRootList())
-					{
-						// only return the first of the root nodes that's in the sorted root list in action domain order.
-						if (RootNode->IsReceivingInput() && RootNode->DoesWidgetSupportActivationFocus())
-						{
-							return RootNode->GatherScrollRecipients();
-						}
-					}
-				}
-			}
-		}
+		return RootNode->GatherScrollRecipients();
 	}
 
 	return TArray<const UWidget*>();
@@ -774,35 +758,12 @@ void UCommonUIActionRouterBase::HandleRootNodeActivated(TWeakPtr<FActivatableTre
 		{
 			ActivatedRoot->SetCanReceiveInput(true);
 			ActivatedRoot->OnLeafmostActiveNodeChanged.BindUObject(this, &UCommonUIActionRouterBase::HandleLeafmostActiveNodeChanged);
-			
-			if (const UCommonInputActionDomainTable* ActionDomainTable = GetActionDomainTable())
+
+			if (!ActiveRootNode.IsValid() && ActivatedRoot == FindActiveActionDomainRootNode())
 			{
-				// We find the first root node that is receiving input and bail early. 
-				// The action domains and root lists are sorted so we should end up with a node with a higher paint layer
-				// or we just stop at the activated root node and make sure it's leaf node is updated.
-				for (const UCommonInputActionDomain* ActionDomain : ActionDomainTable->ActionDomains)
-				{
-					if (FActionDomainSortedRootList* SortedRootList = ActionDomainRootNodes.Find(ActionDomain))
-					{
-						for (FActivatableTreeRootRef& RootNode : SortedRootList->GetRootList())
-						{
-							if (RootNode->IsReceivingInput())
-							{
-								if (RootNode == ActivatedRoot && !ActiveRootNode.IsValid())
-								{
-									if (RootNode->DoesWidgetSupportActivationFocus())
-									{
-										// This will allow us to update the leaf node config on the activated root.
-										ActivatedRoot->UpdateLeafNode();
-									}
-
-									OnBoundActionsUpdated().Broadcast();
-								}
-								break;
-							}
-						}
-					}
-				}
+				// This will allow us to update the leaf node config on the activated root.
+				ActivatedRoot->UpdateLeafNode();
+				OnBoundActionsUpdated().Broadcast();
 			}
 		}
 	}
@@ -1488,7 +1449,6 @@ void UCommonUIActionRouterBase::SetActiveUIInputConfig(const FUIInputConfig& New
 void UCommonUIActionRouterBase::RefreshActionDomainLeafNodeConfig()
 {
 	// We don't want to refresh if there is an activated root node as we don't want input mode changes
-	bool bActivatedRootNodeExists = false;
 	for (const FActivatableTreeRootRef& Root : RootNodes)
 	{
 		if (Root->IsWidgetActivated())
@@ -1499,27 +1459,17 @@ void UCommonUIActionRouterBase::RefreshActionDomainLeafNodeConfig()
 
 	if (const UCommonInputActionDomainTable* ActionDomainTable = GetActionDomainTable())
 	{
-		for (const UCommonInputActionDomain* ActionDomain : ActionDomainTable->ActionDomains)
+		if (FActivatableTreeRootPtr RootNode = FindActiveActionDomainRootNode())
 		{
-			if (FActionDomainSortedRootList* SortedRootList = ActionDomainRootNodes.Find(ActionDomain))
+			if (!RootNode->UpdateLeafmostActiveNode(RootNode))
 			{
-				for (FActivatableTreeRootRef& RootNode : SortedRootList->GetRootList())
-				{
-					// only root nodes that are actively receiving input and supports widget activation focus
-					// should update leaf nodes and have input config applied. This will also update focus.
-					if (RootNode->IsReceivingInput() && RootNode->DoesWidgetSupportActivationFocus())
-					{
-						if (!RootNode->UpdateLeafmostActiveNode(RootNode))
-						{
-							RootNode->ApplyLeafmostNodeConfig();
-						}
-						return;
-					}
-				}
+				RootNode->ApplyLeafmostNodeConfig();
 			}
 		}
-
-		SetActiveUIInputConfig(FUIInputConfig(ActionDomainTable->InputMode, ActionDomainTable->MouseCaptureMode), ActionDomainTable);
+		else
+		{
+			SetActiveUIInputConfig(FUIInputConfig(ActionDomainTable->InputMode, ActionDomainTable->MouseCaptureMode), ActionDomainTable);
+		}
 	}
 }
 
@@ -1673,6 +1623,28 @@ void UCommonUIActionRouterBase::SetActiveActivationMetadata(const FActivationMet
 	OnActivationMetadataChanged().Broadcast(NewConfig);
 }
 
+FActivatableTreeRootPtr UCommonUIActionRouterBase::FindActiveActionDomainRootNode() const
+{
+	if (const UCommonInputActionDomainTable* ActionDomainTable = GetActionDomainTable())
+	{
+		for (const UCommonInputActionDomain* ActionDomain : ActionDomainTable->ActionDomains)
+		{
+			if (const FActionDomainSortedRootList* SortedRootList = ActionDomainRootNodes.Find(ActionDomain))
+			{
+				for (const FActivatableTreeRootRef& RootNode : SortedRootList->GetRootList())
+				{
+					if (RootNode->IsReceivingInput() && RootNode->DoesWidgetSupportActivationFocus())
+					{
+						return RootNode;
+					}
+				}
+			}
+		}
+	}
+
+	return nullptr;
+}
+
 void UCommonUIActionRouterBase::HandleActivatableWidgetRebuilding(UCommonActivatableWidget& RebuildingWidget)
 {
 	if (RebuildingWidget.GetOwningLocalPlayer() == GetLocalPlayerChecked())
diff --git a/Engine/Plugins/Runtime/CommonUI/Source/CommonUI/Private/Input/UIActionRouterTypes.cpp b/Engine/Plugins/Runtime/CommonUI/Source/CommonUI/Private/Input/UIActionRouterTypes.cpp
--- a/Engine/Plugins/Runtime/CommonUI/Source/CommonUI/Private/Input/UIActionRouterTypes.cpp
+++ b/Engine/Plugins/Runtime/CommonUI/Source/CommonUI/Private/Input/UIActionRouterTypes.cpp
@@ -1470,6 +1470,12 @@ void FActivatableTreeRoot::ApplyLeafmostNodeConfig()
 		return;
 	}
 #endif
+	if (!CanSetInputConfigAndFocus())
+	{
+		UE_LOG(LogUIActionRouter, Log, TEXT("Didn't apply input config for leaf-most node of root node [%s] because it's not the active action domain node"), *GetNameSafe(GetWidget()));
+		return;
+	}
+
 	if (FActivatableTreeNodePtr PinnedLeafmostNode = LeafmostActiveNode.Pin())
 	{
 		GetActionRouter().SetActiveActivationMetadata(PinnedLeafmostNode->FindActivationMetadata());
@@ -1498,8 +1504,37 @@ void FActivatableTreeRoot::ApplyLeafmostNodeConfig()
 	}
 }
 
+bool FActivatableTreeRoot::IsAnActionDomainRoot() const
+{
+	for (const TPair<TObjectPtr<UCommonInputActionDomain>, UCommonUIActionRouterBase::FActionDomainSortedRootList>& Pair : GetActionRouter().ActionDomainRootNodes)
+	{
+		if (Pair.Value.Contains(SharedThis(const_cast<FActivatableTreeRoot*>(this))))
+		{
+			return true;
+		}
+	}
+
+	return false;
+}
+
+bool FActivatableTreeRoot::IsActiveActionDomainRoot() const
+{
+	return GetActionRouter().FindActiveActionDomainRootNode().Get() == this;
+}
+
+bool FActivatableTreeRoot::CanSetInputConfigAndFocus() const
+{
+	return !IsAnActionDomainRoot() || IsActiveActionDomainRoot();
+}
+
 void FActivatableTreeRoot::FocusLeafmostNode()
 {
+	if (!CanSetInputConfigAndFocus())
+	{
+		UE_LOG(LogUIActionRouter, Log, TEXT("%s cannot focus its leaf most node because another node has priority"), *GetNameSafe(GetWidget()));
+		return;
+	}
+
 	check(LeafmostActiveNode.IsValid());
 	if (!LeafmostActiveNode.IsValid())
 	{
diff --git a/Engine/Plugins/Runtime/CommonUI/Source/CommonUI/Private/Input/UIActionRouterTypes.h b/Engine/Plugins/Runtime/CommonUI/Source/CommonUI/Private/Input/UIActionRouterTypes.h
--- a/Engine/Plugins/Runtime/CommonUI/Source/CommonUI/Private/Input/UIActionRouterTypes.h
+++ b/Engine/Plugins/Runtime/CommonUI/Source/CommonUI/Private/Input/UIActionRouterTypes.h
@@ -201,6 +201,10 @@ public:
 
 	void ApplyLeafmostNodeConfig();
 
+	bool IsAnActionDomainRoot() const;
+	bool IsActiveActionDomainRoot() const;
+	bool CanSetInputConfigAndFocus() const;
+
 protected:
 	virtual void Init() override;
 
diff --git a/Engine/Plugins/Runtime/CommonUI/Source/CommonUI/Public/Input/CommonUIActionRouterBase.h b/Engine/Plugins/Runtime/CommonUI/Source/CommonUI/Public/Input/CommonUIActionRouterBase.h
--- a/Engine/Plugins/Runtime/CommonUI/Source/CommonUI/Public/Input/CommonUIActionRouterBase.h
+++ b/Engine/Plugins/Runtime/CommonUI/Source/CommonUI/Public/Input/CommonUIActionRouterBase.h
@@ -175,6 +175,10 @@ private:
 	FActivatableTreeNodePtr FindNodeRecursive(const FActivatableTreeNodePtr& CurrentNode, const UCommonActivatableWidget& Widget) const;
 	FActivatableTreeNodePtr FindNodeRecursive(const FActivatableTreeNodePtr& CurrentNode, const TSharedPtr<SWidget>& Widget) const;
 	void SetActiveActivationMetadata(const FActivationMetadata& NewConfig);
+
+	// Returns the top-most active root that is receiving input
+	// Action Domain's analogue to ActiveRootNode
+	FActivatableTreeRootPtr FindActiveActionDomainRootNode() const;
 	
 	void HandleActivatableWidgetRebuilding(UCommonActivatableWidget& RebuildingWidget);
 	void ProcessRebuiltWidgets();
