Add Input Mode gameplay tags
Clear and set enhanced input filter gameplay tags when input mode changes

#tests PIE
#rnx

diff --git a/Engine/Plugins/Runtime/CommonUI/Source/CommonUI/Private/Input/CommonUIActionRouterBase.cpp b/Engine/Plugins/Runtime/CommonUI/Source/CommonUI/Private/Input/CommonUIActionRouterBase.cpp
--- a/Engine/Plugins/Runtime/CommonUI/Source/CommonUI/Private/Input/CommonUIActionRouterBase.cpp
+++ b/Engine/Plugins/Runtime/CommonUI/Source/CommonUI/Private/Input/CommonUIActionRouterBase.cpp
@@ -11,6 +11,7 @@
 #include "Engine/Console.h"
 #include "Engine/Engine.h"
 #include "Engine/GameInstance.h"
+#include "EnhancedInputSubsystems.h"
 #include "Framework/Application/SlateApplication.h"
 #include "Framework/Application/SlateUser.h"
 #include "Framework/Commands/InputBindingManager.h"
@@ -19,6 +20,7 @@
 #include "Input/CommonUIInputSettings.h"
 #include "Input/UIActionBinding.h"
 #include "Input/UIActionRouterTypes.h"
+#include "NativeGameplayTags.h"
 #include "Slate/SGameLayerManager.h"
 #include "Slate/SObjectWidget.h"
 #include "TimerManager.h"
@@ -81,6 +83,11 @@ static FAutoConsoleVariableRef CVarCheckGameViewportClientValid(
 
 //////////////////////////////////////////////////////////////////////////
 
+UE_DEFINE_GAMEPLAY_TAG_STATIC(TAG_InputModeGame, "InputMode.Game");
+UE_DEFINE_GAMEPLAY_TAG_STATIC(TAG_InputModeMenu, "InputMode.Menu");
+
+//////////////////////////////////////////////////////////////////////////
+
 FGlobalUITags FGlobalUITags::GUITags;
 
 UCommonActivatableWidget* FindOwningActivatable(const UWidget& Widget)
@@ -930,6 +937,26 @@ EProcessHoldActionResult UCommonUIActionRouterBase::ProcessHoldInputOnActionDoma
 	return HoldActionResult;
 }
 
+FGameplayTagContainer UCommonUIActionRouterBase::GetGameplayTagsForInputMode(const ECommonInputMode Mode) const
+{
+	FGameplayTagContainer Tags;
+	
+	switch (Mode)
+	{
+	case ECommonInputMode::Game:
+		Tags.AddTag(TAG_InputModeGame);
+
+	case ECommonInputMode::Menu:
+		Tags.AddTag(TAG_InputModeMenu);
+
+	case ECommonInputMode::All:
+		Tags.AddTag(TAG_InputModeGame);
+		Tags.AddTag(TAG_InputModeMenu);
+	}
+
+	return Tags;
+}
+
 void UCommonUIActionRouterBase::ProcessRebuiltWidgets()
 {
 	// Begin by organizing all of the widgets that need nodes according to their direct parent
@@ -1611,6 +1638,11 @@ void UCommonUIActionRouterBase::ApplyUIInputConfig(const FUIInputConfig& NewConf
 
 		if (PreviousInputMode != NewConfig.GetInputMode())
 		{
+			if (UEnhancedInputLocalPlayerSubsystem* IE = LocalPlayer.GetSubsystem<UEnhancedInputLocalPlayerSubsystem>())
+			{
+				IE->RemoveTagsFromInputMode(GetGameplayTagsForInputMode(PreviousInputMode));
+				IE->AppendTagsToInputMode(GetGameplayTagsForInputMode(NewConfig.GetInputMode()));
+			}
 			OnActiveInputModeChanged().Broadcast(NewConfig.GetInputMode());
 		}
 
diff --git a/Engine/Plugins/Runtime/CommonUI/Source/CommonUI/Public/Input/CommonUIActionRouterBase.h b/Engine/Plugins/Runtime/CommonUI/Source/CommonUI/Public/Input/CommonUIActionRouterBase.h
--- a/Engine/Plugins/Runtime/CommonUI/Source/CommonUI/Public/Input/CommonUIActionRouterBase.h
+++ b/Engine/Plugins/Runtime/CommonUI/Source/CommonUI/Public/Input/CommonUIActionRouterBase.h
@@ -12,6 +12,7 @@
 #include "InputCoreTypes.h"
 #include "UObject/WeakObjectPtr.h"
 #include "Framework/Application/SlateApplication.h"
+#include "GameplayTagContainer.h"
 #include "Misc/FrameValue.h"
 #include "CommonUIActionRouterBase.generated.h"
 
@@ -196,6 +197,7 @@ private:
 	const UCommonInputActionDomainTable* GetActionDomainTable() const;
 	bool ProcessInputOnActionDomains(ECommonInputMode ActiveInputMode, FKey Key, EInputEvent InputEvent) const;
 	EProcessHoldActionResult ProcessHoldInputOnActionDomains(ECommonInputMode ActiveInputMode, FKey Key, EInputEvent InputEvent) const;
+	FGameplayTagContainer GetGameplayTagsForInputMode(const ECommonInputMode Mode) const;
 
 	struct FPendingWidgetRegistration
 	{
