CommonUI: Don't broadcast input method changed events when EI rebuilds its input mappings and instead bind directly to `ControlMappingsRebuiltDelegate` in CommonActionWidget to update itself in case the keys changed

- This had unwanted side-effects, such as refreshing focus for the active leaf mode node which is not generally influenced by EI's mappings

diff --git a/Engine/Plugins/Runtime/CommonUI/Source/CommonInput/Private/CommonInputSubsystem.cpp b/Engine/Plugins/Runtime/CommonUI/Source/CommonInput/Private/CommonInputSubsystem.cpp
--- a/Engine/Plugins/Runtime/CommonUI/Source/CommonInput/Private/CommonInputSubsystem.cpp
+++ b/Engine/Plugins/Runtime/CommonUI/Source/CommonInput/Private/CommonInputSubsystem.cpp
@@ -89,18 +89,6 @@ void UCommonInputSubsystem::Initialize(FSubsystemCollectionBase& Collection)
 
 	CVarInputKeysVisible->SetOnChangedCallback(FConsoleVariableDelegate::CreateUObject(this, &UCommonInputSubsystem::ShouldShowInputKeysChanged));
 
-	if (ICommonInputModule::Get().GetSettings().GetEnableEnhancedInputSupport())
-	{
-		if (ULocalPlayer* LocalPlayer = GetLocalPlayerChecked())
-		{
-			if (UEnhancedInputLocalPlayerSubsystem* EnhancedInputLocalPlayerSubsystem = LocalPlayer->GetSubsystem<UEnhancedInputLocalPlayerSubsystem>())
-			{
-				BroadcastInputMethodChangedEvent.BindUFunction(this, GET_FUNCTION_NAME_CHECKED(UCommonInputSubsystem, BroadcastInputMethodChanged));
-				EnhancedInputLocalPlayerSubsystem->ControlMappingsRebuiltDelegate.AddUnique(BroadcastInputMethodChangedEvent);
-			}
-		}
-	}
-
 	SetActionDomainTable(FCommonInputBase::GetInputSettings()->GetActionDomainTable());
 }
 
@@ -113,18 +101,6 @@ void UCommonInputSubsystem::Deinitialize()
 	}
 	CommonInputPreprocessor.Reset();
 
-	if (ICommonInputModule::Get().GetSettings().GetEnableEnhancedInputSupport())
-	{
-		if (ULocalPlayer* LocalPlayer = GetLocalPlayerChecked())
-		{
-			if (UEnhancedInputLocalPlayerSubsystem* EnhancedInputLocalPlayerSubsystem = LocalPlayer->GetSubsystem<UEnhancedInputLocalPlayerSubsystem>())
-			{
-				EnhancedInputLocalPlayerSubsystem->ControlMappingsRebuiltDelegate.Remove(BroadcastInputMethodChangedEvent);
-				BroadcastInputMethodChangedEvent.Unbind();
-			}
-		}
-	}
-
 	FTSTicker::GetCoreTicker().RemoveTicker(TickHandle);
 }
 
diff --git a/Engine/Plugins/Runtime/CommonUI/Source/CommonInput/Public/CommonInputSubsystem.h b/Engine/Plugins/Runtime/CommonUI/Source/CommonInput/Public/CommonInputSubsystem.h
--- a/Engine/Plugins/Runtime/CommonUI/Source/CommonInput/Public/CommonInputSubsystem.h
+++ b/Engine/Plugins/Runtime/CommonUI/Source/CommonInput/Public/CommonInputSubsystem.h
@@ -170,7 +170,4 @@ private:
 	* Note : Calling order is not guaranteed. Also, keep in mind that you might need to honor the previous callee's request to not support the input type being tested.
 	*/
 	static FPlatformInputSupportOverrideDelegate OnPlatformInputSupportOverride;
-
-	/** Used to broadcast when input method should be considered changed based on external systems, such as on enhanced input mapping context application*/
-	FScriptDelegate BroadcastInputMethodChangedEvent;
 };
\ No newline at end of file
diff --git a/Engine/Plugins/Runtime/CommonUI/Source/CommonUI/Private/CommonActionWidget.cpp b/Engine/Plugins/Runtime/CommonUI/Source/CommonUI/Private/CommonActionWidget.cpp
--- a/Engine/Plugins/Runtime/CommonUI/Source/CommonUI/Private/CommonActionWidget.cpp
+++ b/Engine/Plugins/Runtime/CommonUI/Source/CommonUI/Private/CommonActionWidget.cpp
@@ -7,6 +7,8 @@
 #include "CommonInputTypeEnum.h"
 #include "CommonUITypes.h"
 #include "CommonWidgetPaletteCategories.h"
+#include "Engine/LocalPlayer.h"
+#include "EnhancedInputSubsystems.h"
 #include "Input/UIActionBinding.h"
 #include "InputAction.h"
 #include "InputTriggers.h"
@@ -417,6 +419,23 @@ void UCommonActionWidget::ListenToInputMethodChanged(bool bListen)
 			CommonInputSubsystem->OnInputMethodChangedNative.AddUObject(this, &ThisClass::HandleInputMethodChanged);
 		}
 	}
+
+	if (CommonUI::IsEnhancedInputSupportEnabled())
+	{
+		const UWidget* BoundWidget = DisplayedBindingHandle.GetBoundWidget();
+		const ULocalPlayer* LocalPlayer = BoundWidget ? BoundWidget->GetOwningLocalPlayer() : GetOwningLocalPlayer();
+		if (UEnhancedInputLocalPlayerSubsystem* EnhancedInputLocalPlayerSubsystem = ULocalPlayer::GetSubsystem<UEnhancedInputLocalPlayerSubsystem>(LocalPlayer))
+		{
+			if (bListen)
+			{
+				EnhancedInputLocalPlayerSubsystem->ControlMappingsRebuiltDelegate.AddUniqueDynamic(this, &UCommonActionWidget::OnEnhancedInputMappingsRebuilt);
+			}
+			else
+			{
+				EnhancedInputLocalPlayerSubsystem->ControlMappingsRebuiltDelegate.RemoveDynamic(this, &UCommonActionWidget::OnEnhancedInputMappingsRebuilt);
+			}
+		}
+	}
 }
 
 void UCommonActionWidget::HandleInputMethodChanged(ECommonInputType InInputType)
@@ -425,6 +444,11 @@ void UCommonActionWidget::HandleInputMethodChanged(ECommonInputType InInputType)
 	OnInputMethodChanged.Broadcast(InInputType==ECommonInputType::Gamepad);
 }
 
+void UCommonActionWidget::OnEnhancedInputMappingsRebuilt()
+{
+	UpdateActionWidget();
+}
+
 #if WITH_EDITOR
 const FText UCommonActionWidget::GetPaletteCategory()
 {
diff --git a/Engine/Plugins/Runtime/CommonUI/Source/CommonUI/Public/CommonActionWidget.h b/Engine/Plugins/Runtime/CommonUI/Source/CommonUI/Public/CommonActionWidget.h
--- a/Engine/Plugins/Runtime/CommonUI/Source/CommonUI/Public/CommonActionWidget.h
+++ b/Engine/Plugins/Runtime/CommonUI/Source/CommonUI/Public/CommonActionWidget.h
@@ -137,6 +137,9 @@ protected:
 	void ListenToInputMethodChanged(bool bListen = true);
 	void HandleInputMethodChanged(ECommonInputType InInputType);
 
+	UFUNCTION()
+	void OnEnhancedInputMappingsRebuilt();
+
 	UPROPERTY(Transient)
 	TObjectPtr<UMaterialInstanceDynamic> ProgressDynamicMaterial;
 	
