CommonUI: Fix action domain root nodes not receiving focus/applying their input config if they are activated before painting

- Also sort action domain root nodes on tick instead of once per frame to better match non-action domain nodes

diff --git a/Engine/Plugins/Runtime/CommonUI/Source/CommonUI/Private/Input/CommonUIActionRouterBase.cpp b/Engine/Plugins/Runtime/CommonUI/Source/CommonUI/Private/Input/CommonUIActionRouterBase.cpp
--- a/Engine/Plugins/Runtime/CommonUI/Source/CommonUI/Private/Input/CommonUIActionRouterBase.cpp
+++ b/Engine/Plugins/Runtime/CommonUI/Source/CommonUI/Private/Input/CommonUIActionRouterBase.cpp
@@ -745,16 +745,19 @@ void UCommonUIActionRouterBase::HandleRootNodeActivated(TWeakPtr<FActivatableTre
 	FActivatableTreeRootRef ActivatedRoot = WeakActivatedRoot.Pin().ToSharedRef();
 	UCommonActivatableWidget* NodeWidget = ActivatedRoot->GetWidget();
 
-	if (RootNodes.Contains(ActivatedRoot) && ActivatedRoot->GetLastPaintLayer() > 0)
+	if (RootNodes.Contains(ActivatedRoot))
 	{
-		const int32 CurrentRootLayer = ActiveRootNode ? ActiveRootNode->GetLastPaintLayer() : INDEX_NONE;
-		if (ActivatedRoot->GetLastPaintLayer() > CurrentRootLayer)
+		if (ActivatedRoot->GetLastPaintLayer() > 0)
 		{
-			// Ensure we have a local player so the action router local player subsystem will handle root change
-			const ULocalPlayer* LocalPlayer = GetLocalPlayer();
-			if (LocalPlayer && LocalPlayer->ViewportClient)
+			const int32 CurrentRootLayer = ActiveRootNode ? ActiveRootNode->GetLastPaintLayer() : INDEX_NONE;
+			if (ActivatedRoot->GetLastPaintLayer() > CurrentRootLayer)
 			{
-				SetActiveRoot(ActivatedRoot);
+				// Ensure we have a local player so the action router local player subsystem will handle root change
+				const ULocalPlayer* LocalPlayer = GetLocalPlayer();
+				if (LocalPlayer && LocalPlayer->ViewportClient)
+				{
+					SetActiveRoot(ActivatedRoot);
+				}
 			}
 		}
 	}
@@ -766,11 +769,21 @@ void UCommonUIActionRouterBase::HandleRootNodeActivated(TWeakPtr<FActivatableTre
 			ActivatedRoot->SetCanReceiveInput(true);
 			ActivatedRoot->OnLeafmostActiveNodeChanged.BindUObject(this, &UCommonUIActionRouterBase::HandleLeafmostActiveNodeChanged);
 
-			if (!ActiveRootNode.IsValid() && ActivatedRoot == FindActiveActionDomainRootNode())
+			if (!ActiveRootNode.IsValid())
 			{
-				// This will allow us to update the leaf node config on the activated root.
-				ActivatedRoot->UpdateLeafNode();
-				OnBoundActionsUpdated().Broadcast();
+				if (ActivatedRoot->GetLastPaintLayer() > 0)
+				{
+					if (ActivatedRoot == FindActiveActionDomainRootNode())
+					{
+						// This will allow us to update the leaf node config on the activated root.
+						ActivatedRoot->UpdateLeafNode();
+						OnBoundActionsUpdated().Broadcast();
+					}
+				}
+				else
+				{
+					ActiveActionDomainRootsPendingPaint.Add(ActivatedRoot);
+				}
 			}
 		}
 	}
@@ -1080,6 +1093,13 @@ void UCommonUIActionRouterBase::AssembleTreeRecursive(const FActivatableTreeNode
 bool UCommonUIActionRouterBase::Tick(float DeltaTime)
 {
 	QUICK_SCOPE_CYCLE_COUNTER(STAT_UCommonUIActionRouter_Tick);
+
+	// Sort our action domain roots to match the most recent paint layers
+	for (TPair<TObjectPtr<UCommonInputActionDomain>, FActionDomainSortedRootList>& Pair : ActionDomainRootNodes)
+	{
+		Pair.Value.Sort();
+	}
+	
 	if (PendingWidgetRegistrations.Num() > 0 || RebuiltWidgetsPendingNodeAssignment.Num() > 0)
 	{
 		ProcessRebuiltWidgets();
@@ -1106,6 +1126,33 @@ bool UCommonUIActionRouterBase::Tick(float DeltaTime)
 		{
 			SetActiveRoot(NewActiveRoot);
 		}
+
+		// Check for any newly painted active action domain nodes that could become the active node
+		// Work on a copy in-case ActiveActionDomainRootsPendingPaint changes during iteration
+		// We iterate even if ActiveRootNode.IsValid() to maintain the list of unpainted widgets that could become the active node if ActiveRootNode deactivates before then
+		TSet<TWeakPtr<FActivatableTreeRoot>> ActiveActionDomainRootsPendingPaintCopy = MoveTemp(ActiveActionDomainRootsPendingPaint);
+		for (auto It = ActiveActionDomainRootsPendingPaintCopy.CreateIterator(); It; ++It)
+		{
+			if (FActivatableTreeRootPtr ActiveRoot = It->Pin())
+			{
+				if (ActiveRoot->GetLastPaintLayer() > 0)
+				{
+					It.RemoveCurrent();
+					
+					if (!ActiveRootNode.IsValid() && ActiveRoot == FindActiveActionDomainRootNode())
+					{
+						ActiveRoot->UpdateLeafNode();
+						OnBoundActionsUpdated().Broadcast();
+					}
+				}
+			}
+			else
+			{
+				It.RemoveCurrent();
+			}
+		}
+
+		ActiveActionDomainRootsPendingPaint.Append(ActiveActionDomainRootsPendingPaintCopy);
 	}
 
 	const ECommonInputMode ActiveMode = GetActiveInputMode();
@@ -1832,6 +1879,5 @@ void UCommonUIActionRouterBase::FActionDomainSortedRootList::Sort()
 	};
 
 	RootList.Sort(SortFunc);
-	bDidSortThisFrame = true;
 }
 
diff --git a/Engine/Plugins/Runtime/CommonUI/Source/CommonUI/Public/Input/CommonUIActionRouterBase.h b/Engine/Plugins/Runtime/CommonUI/Source/CommonUI/Public/Input/CommonUIActionRouterBase.h
--- a/Engine/Plugins/Runtime/CommonUI/Source/CommonUI/Public/Input/CommonUIActionRouterBase.h
+++ b/Engine/Plugins/Runtime/CommonUI/Source/CommonUI/Public/Input/CommonUIActionRouterBase.h
@@ -13,7 +13,6 @@
 #include "UObject/WeakObjectPtr.h"
 #include "Framework/Application/SlateApplication.h"
 #include "GameplayTagContainer.h"
-#include "Misc/FrameValue.h"
 #include "CommonUIActionRouterBase.generated.h"
 
 class SWidget;
@@ -230,6 +229,7 @@ PRAGMA_ENABLE_DEPRECATION_WARNINGS
 	TArray<FPendingWidgetRegistration> PendingWidgetRegistrations;
 	
 	TArray<TWeakObjectPtr<UCommonActivatableWidget>> RebuiltWidgetsPendingNodeAssignment;
+	TSet<TWeakPtr<FActivatableTreeRoot>> ActiveActionDomainRootsPendingPaint;
 	
 	TArray<FActivatableTreeRootRef> RootNodes;
 	FActivatableTreeRootPtr ActiveRootNode;
@@ -259,20 +259,10 @@ PRAGMA_ENABLE_DEPRECATION_WARNINGS
 	{
 		TArray<FActivatableTreeRootRef>& GetRootList()
 		{
-			if (!bDidSortThisFrame.TryGetValue(false))
-			{
-				// Make sure the list sorting is always up to date
-				Sort();
-			}
 			return RootList;
 		}
 		const TArray<FActivatableTreeRootRef>& GetRootList() const
 		{
-			if (!bDidSortThisFrame.TryGetValue(false))
-			{
-				// Make sure the list sorting is always up to date
-				const_cast<FActionDomainSortedRootList*>(this)->Sort();
-			}
 			return RootList;
 		}
 		
@@ -289,7 +279,6 @@ PRAGMA_ENABLE_DEPRECATION_WARNINGS
 
 	private:
 		TArray<FActivatableTreeRootRef> RootList;
-		TFrameValue<bool> bDidSortThisFrame;
 	};
 
 	TMap<TObjectPtr<UCommonInputActionDomain>, FActionDomainSortedRootList> ActionDomainRootNodes;
