Common UI: add a priority setting for action bindings on the same widget in case we want to execute them in a different order than the order of registration.

diff --git a/Engine/Plugins/Runtime/CommonUI/Source/CommonUI/Private/Input/UIActionRouterTypes.cpp b/Engine/Plugins/Runtime/CommonUI/Source/CommonUI/Private/Input/UIActionRouterTypes.cpp
--- a/Engine/Plugins/Runtime/CommonUI/Source/CommonUI/Private/Input/UIActionRouterTypes.cpp
+++ b/Engine/Plugins/Runtime/CommonUI/Source/CommonUI/Private/Input/UIActionRouterTypes.cpp
@@ -109,6 +109,7 @@ FUIActionBinding::FUIActionBinding(const UWidget& InBoundWidget, const FBindUIAc
 	, InputEvent(BindArgs.KeyEvent)
 	, bConsumesInput(BindArgs.bConsumeInput)
 	, bIsPersistent(BindArgs.bIsPersistent)
+	, PriorityWithinCollection(BindArgs.PriorityWithinCollection)
 	, BoundWidget(&InBoundWidget)
 	, InputMode(BindArgs.InputMode)
 	, bDisplayInActionBar(BindArgs.bDisplayInActionBar)
@@ -806,7 +807,21 @@ void FActionRouterBindingCollection::AddBinding(FUIActionBinding& Binding)
 {
 	if (ensure(!ActionBindings.Contains(Binding.Handle)))
 	{
-		ActionBindings.Add(Binding.Handle);
+		int32 IndexToInsertBinding = ActionBindings.Num();
+		for (int32 Index = ActionBindings.Num() - 1; Index >= 0; Index--)
+		{
+			if (TSharedPtr<const FUIActionBinding> UIActionBinding = FUIActionBinding::FindBinding(ActionBindings[Index]))
+			{
+				if (UIActionBinding->PriorityWithinCollection < Binding.PriorityWithinCollection)
+				{
+					IndexToInsertBinding = Index;
+					continue;
+				}
+			}
+			break;
+		}
+		ActionBindings.Insert(Binding.Handle, IndexToInsertBinding);
+
 		Binding.OwningCollection = AsShared();
 
 		if (Binding.HoldMappings.Num() > 0)
diff --git a/Engine/Plugins/Runtime/CommonUI/Source/CommonUI/Public/Input/CommonUIInputTypes.h b/Engine/Plugins/Runtime/CommonUI/Source/CommonUI/Public/Input/CommonUIInputTypes.h
--- a/Engine/Plugins/Runtime/CommonUI/Source/CommonUI/Public/Input/CommonUIInputTypes.h
+++ b/Engine/Plugins/Runtime/CommonUI/Source/CommonUI/Public/Input/CommonUIInputTypes.h
@@ -81,6 +81,12 @@ struct COMMONUI_API FBindUIActionArgs
 	/** Optional display name to associate with this binding instead of the default */
 	FText OverrideDisplayName;
 
+	/**
+	 * Normally, actions on a widget are triggered in the order they're registered. We can assign a priority to ensure a certain order of execution.
+	 * 0 is the order of registration.
+	 */
+	int32 PriorityWithinCollection = 0;
+
 	FSimpleDelegate OnExecuteAction;
 
 	/** If the bound action has any hold mappings, this will fire each frame while held. Has no bearing on actual execution and wholly irrelevant for non-hold actions */
diff --git a/Engine/Plugins/Runtime/CommonUI/Source/CommonUI/Public/Input/UIActionBinding.h b/Engine/Plugins/Runtime/CommonUI/Source/CommonUI/Public/Input/UIActionBinding.h
--- a/Engine/Plugins/Runtime/CommonUI/Source/CommonUI/Public/Input/UIActionBinding.h
+++ b/Engine/Plugins/Runtime/CommonUI/Source/CommonUI/Public/Input/UIActionBinding.h
@@ -56,6 +56,8 @@ struct COMMONUI_API FUIActionBinding
 	bool bConsumesInput = true;
 	bool bIsPersistent = false;
 	
+	int32 PriorityWithinCollection = 0;
+
 	TWeakObjectPtr<const UWidget> BoundWidget;
 	ECommonInputMode InputMode;
 
