Removing ensure, it's no longer true that this function can be called only when common input subsystem exists

diff --git a/Engine/Plugins/Runtime/CommonUI/Source/CommonUI/Private/CommonActionWidget.cpp b/Engine/Plugins/Runtime/CommonUI/Source/CommonUI/Private/CommonActionWidget.cpp
--- a/Engine/Plugins/Runtime/CommonUI/Source/CommonUI/Private/CommonActionWidget.cpp
+++ b/Engine/Plugins/Runtime/CommonUI/Source/CommonUI/Private/CommonActionWidget.cpp
@@ -331,7 +331,7 @@ void UCommonActionWidget::UpdateActionWidget()
 	if (GetWorld())
 	{
 		const UCommonInputSubsystem* CommonInputSubsystem = GetInputSubsystem();
-		if (IsDesignTime() || (GetGameInstance() && ensure(CommonInputSubsystem) && CommonInputSubsystem->ShouldShowInputKeys()))
+		if (IsDesignTime() || (GetGameInstance() && CommonInputSubsystem && CommonInputSubsystem->ShouldShowInputKeys()))
 		{
 			if (ShouldUpdateActionWidgetIcon())
 			{
