CommonUI: Centralize handling the changing of the CurrentInputType

- CommonUI's tracking of the user's current input type is now encapsulated by `UCommonInputSubsystem::RecalculateCurrentInputType()`
- Rename `LastInputType` to `RawInputType` to clarify what it is
- Add comments to `RawInputType` and `CurrentInputType`

#rnx

diff --git a/Engine/Plugins/Runtime/CommonUI/Source/CommonInput/Private/CommonInputSubsystem.cpp b/Engine/Plugins/Runtime/CommonUI/Source/CommonInput/Private/CommonInputSubsystem.cpp
--- a/Engine/Plugins/Runtime/CommonUI/Source/CommonInput/Private/CommonInputSubsystem.cpp
+++ b/Engine/Plugins/Runtime/CommonUI/Source/CommonInput/Private/CommonInputSubsystem.cpp
@@ -77,7 +77,8 @@ void UCommonInputSubsystem::Initialize(FSubsystemCollectionBase& Collection)
 	const UCommonInputPlatformSettings* Settings = UPlatformSettingsManager::Get().GetSettingsForPlatform<UCommonInputPlatformSettings>();
 
 	GamepadInputType = Settings->GetDefaultGamepadName();
-	CurrentInputType = LastInputType = Settings->GetDefaultInputType();
+	RawInputType = Settings->GetDefaultInputType();
+	CurrentInputType = RawInputType;
 
 	CommonInputPreprocessor = MakeInputProcessor();
 	if (FSlateApplication::IsInitialized())
@@ -137,7 +138,7 @@ void UCommonInputSubsystem::AddOrRemoveInputTypeLock(FName InReason, ECommonInpu
 	}
 
 	int32 ComputedInputLock = INDEX_NONE;
-	for (auto Entry : CurrentInputLocks)
+	for (const TPair<FName, ECommonInputType>& Entry : CurrentInputLocks)
 	{
 		// Take the most restrictive lock, e.g. Gamepad lock is more restrictive than a Keyboard/Mouse lock.
 		if (((int32)Entry.Value) > ComputedInputLock)
@@ -155,15 +156,7 @@ void UCommonInputSubsystem::AddOrRemoveInputTypeLock(FName InReason, ECommonInpu
 		CurrentInputLock = (ECommonInputType)ComputedInputLock;
 	}
 
-	const ECommonInputType PreviousInput = CurrentInputType;
-
-	// If a lock was put in place, lock the current input type.
-	CurrentInputType = LockInput(LastInputType);
-
-	if (CurrentInputType != PreviousInput)
-	{
-		BroadcastInputMethodChanged();
-	}
+	RecalculateCurrentInputType();
 }
 
 bool UCommonInputSubsystem::IsInputMethodActive(ECommonInputType InputMethod) const
@@ -255,68 +248,71 @@ bool UCommonInputSubsystem::CheckForInputMethodThrashing(ECommonInputType NewInp
 	return false;
 }
 
-void UCommonInputSubsystem::SetCurrentInputType(ECommonInputType NewInputType)
+void UCommonInputSubsystem::RecalculateCurrentInputType()
 {
-	if ((LastInputType != NewInputType) && PlatformSupportsInputType(NewInputType))
+	ECommonInputType LockedInput = LockInput(RawInputType);
+
+	if (LockedInput != CurrentInputType)
 	{
-		CheckForInputMethodThrashing(NewInputType);
-	
-		//If we have any locks we can't change the input mode.
-		if (!CurrentInputLocks.Num())
+#if !UE_BUILD_SHIPPING
+		if (bDumpInputTypeChangeCallstack)
 		{
-			LastInputType = NewInputType;
+			const uint32 DumpCallstackSize = 65535;
+			ANSICHAR DumpCallstack[DumpCallstackSize] = { 0 };
+			FString ScriptStack = FFrame::GetScriptCallstack(true /* bReturnEmpty */);
+			FPlatformStackWalk::StackWalkAndDump(DumpCallstack, DumpCallstackSize, 0);
+			UE_LOG(LogCommonInput, Log, TEXT("--- Input Changing Callstack ---"));
+			UE_LOG(LogCommonInput, Log, TEXT("Script Stack:\n%s"), *ScriptStack);
+			UE_LOG(LogCommonInput, Log, TEXT("Callstack:\n%s"), ANSI_TO_TCHAR(DumpCallstack));
+		}
+#endif // !UE_BUILD_SHIPPING
+		
+		CurrentInputType = LockedInput;
 
-			ECommonInputType LockedInput = LockInput(NewInputType);
+		FSlateApplication& SlateApplication = FSlateApplication::Get();
+		ULocalPlayer* LocalPlayer = GetLocalPlayerChecked();
+		bool bCursorUser = LocalPlayer && LocalPlayer->GetSlateUser() == SlateApplication.GetCursorUser();
 
-			if (LockedInput != CurrentInputType)
+		switch (CurrentInputType)
+		{
+		case ECommonInputType::Gamepad:
+			UE_LOG(LogCommonInput, Log, TEXT("UCommonInputSubsystem::RecalculateCurrentInputType(): Using Gamepad"));
+			if (bCursorUser)
 			{
-#if !UE_BUILD_SHIPPING
-				if (bDumpInputTypeChangeCallstack)
-				{
-					const uint32 DumpCallstackSize = 65535;
-					ANSICHAR DumpCallstack[DumpCallstackSize] = { 0 };
-					FString ScriptStack = FFrame::GetScriptCallstack(true /* bReturnEmpty */);
-					FPlatformStackWalk::StackWalkAndDump(DumpCallstack, DumpCallstackSize, 0);
-					UE_LOG(LogCommonInput, Log, TEXT("--- Input Changing Callstack ---"));
-					UE_LOG(LogCommonInput, Log, TEXT("Script Stack:\n%s"), *ScriptStack);
-					UE_LOG(LogCommonInput, Log, TEXT("Callstack:\n%s"), ANSI_TO_TCHAR(DumpCallstack));
-				}
-#endif // !UE_BUILD_SHIPPING
-				
-				CurrentInputType = LockedInput;
-
-				FSlateApplication& SlateApplication = FSlateApplication::Get();
-				ULocalPlayer* LocalPlayer = GetLocalPlayerChecked();
-				bool bCursorUser = LocalPlayer && LocalPlayer->GetSlateUser() == SlateApplication.GetCursorUser();
-
-				switch (CurrentInputType)
-				{
-				case ECommonInputType::Gamepad:
-					UE_LOG(LogCommonInput, Log, TEXT("UCommonInputSubsystem::SetCurrentInputType(): Using Gamepad"));
-					if (bCursorUser)
-					{
-						SlateApplication.UsePlatformCursorForCursorUser(bEnableGamepadPlatformCursor);
-					}
-					SlateApplication.SetGameAllowsFakingTouchEvents(false);
-					break;
-				case ECommonInputType::Touch:
-					UE_LOG(LogCommonInput, Log, TEXT("UCommonInputSubsystem::SetCurrentInputType(): Using Touch"));
-					SlateApplication.SetGameAllowsFakingTouchEvents(true);
-					SlateApplication.SetGameIsFakingTouchEvents(LocalPlayer && LocalPlayer->ViewportClient && LocalPlayer->ViewportClient->GetUseMouseForTouch());
-					break;
-				case ECommonInputType::MouseAndKeyboard:
-				default:				
-					UE_LOG(LogCommonInput, Log, TEXT("UCommonInputSubsystem::SetCurrentInputType(): Using Mouse"));
-					if (bCursorUser)
-					{
-						SlateApplication.UsePlatformCursorForCursorUser(true);
-					}
-					SlateApplication.SetGameAllowsFakingTouchEvents(false);
-					break;
-				}
-
-				BroadcastInputMethodChanged();
+				SlateApplication.UsePlatformCursorForCursorUser(bEnableGamepadPlatformCursor);
 			}
+			SlateApplication.SetGameAllowsFakingTouchEvents(false);
+			break;
+		case ECommonInputType::Touch:
+			UE_LOG(LogCommonInput, Log, TEXT("UCommonInputSubsystem::RecalculateCurrentInputType(): Using Touch"));
+			SlateApplication.SetGameAllowsFakingTouchEvents(true);
+			SlateApplication.SetGameIsFakingTouchEvents(LocalPlayer && LocalPlayer->ViewportClient && LocalPlayer->ViewportClient->GetUseMouseForTouch());
+			break;
+		case ECommonInputType::MouseAndKeyboard:
+		default:				
+			UE_LOG(LogCommonInput, Log, TEXT("UCommonInputSubsystem::RecalculateCurrentInputType(): Using Mouse"));
+			if (bCursorUser)
+			{
+				SlateApplication.UsePlatformCursorForCursorUser(true);
+			}
+			SlateApplication.SetGameAllowsFakingTouchEvents(false);
+			break;
+		}
+
+		BroadcastInputMethodChanged();
+	}
+}
+
+void UCommonInputSubsystem::SetCurrentInputType(ECommonInputType NewInputType)
+{
+	if ((RawInputType != NewInputType) && PlatformSupportsInputType(NewInputType))
+	{
+		RawInputType = NewInputType;
+
+		const bool bIsLockedByThrashing = CheckForInputMethodThrashing(NewInputType);
+		if (!bIsLockedByThrashing)
+		{
+			RecalculateCurrentInputType();
 		}
 	}
 }
@@ -342,7 +338,7 @@ bool UCommonInputSubsystem::IsUsingPointerInput() const
 {
 	bool bUsingPointerInput = false;
 
-	switch (LastInputType)
+	switch (RawInputType)
 	{
 		case ECommonInputType::MouseAndKeyboard:
 		case ECommonInputType::Touch:
diff --git a/Engine/Plugins/Runtime/CommonUI/Source/CommonInput/Public/CommonInputSubsystem.h b/Engine/Plugins/Runtime/CommonUI/Source/CommonInput/Public/CommonInputSubsystem.h
--- a/Engine/Plugins/Runtime/CommonUI/Source/CommonInput/Public/CommonInputSubsystem.h
+++ b/Engine/Plugins/Runtime/CommonUI/Source/CommonInput/Public/CommonInputSubsystem.h
@@ -122,6 +122,8 @@ private:
 
 	bool CheckForInputMethodThrashing(ECommonInputType NewInputType);
 
+	void RecalculateCurrentInputType();
+
 	FTSTicker::FDelegateHandle TickHandle;
 
 	UPROPERTY(BlueprintAssignable, Category = CommonInputSubsystem, meta = (AllowPrivateAccess))
@@ -136,11 +138,11 @@ private:
 	UPROPERTY(Transient)
 	double LastTimeInputMethodThrashingBegan = 0;
 
-	/**  */
+	/** The most recent input type that the user used, before considering locks and thrashing, but does consider PlatformSupportsInputType() */
 	UPROPERTY(Transient)
-	ECommonInputType LastInputType;
+	ECommonInputType RawInputType;
 
-	/**  */
+	/** The current effective input type after considering input locks and thrashing */
 	UPROPERTY(Transient)
 	ECommonInputType CurrentInputType;
 
