Make UCommonInputPlatformSettings clear its cached controller data classes when it is modified via config. Fixes issue where new controller data added via config can cause the cached array to grow forever.

diff --git a/Engine/Plugins/Runtime/CommonUI/Source/CommonInput/Private/CommonInputBaseTypes.cpp b/Engine/Plugins/Runtime/CommonUI/Source/CommonInput/Private/CommonInputBaseTypes.cpp
--- a/Engine/Plugins/Runtime/CommonUI/Source/CommonInput/Private/CommonInputBaseTypes.cpp
+++ b/Engine/Plugins/Runtime/CommonUI/Source/CommonInput/Private/CommonInputBaseTypes.cpp
@@ -405,6 +405,13 @@ void UCommonInputPlatformSettings::PostEditChangeProperty(struct FPropertyChange
 }
 #endif
 
+void UCommonInputPlatformSettings::PostReloadConfig(FProperty* PropertyThatWasLoaded)
+{
+	Super::PostReloadConfig(PropertyThatWasLoaded);
+
+	ControllerDataClasses.Reset();
+}
+
 bool FCommonInputPlatformBaseData::TryGetInputBrush(FSlateBrush& OutBrush, FKey Key, ECommonInputType InputType, const FName& GamepadName) const
 {
 	if (ControllerDataClasses.Num() > 0)
diff --git a/Engine/Plugins/Runtime/CommonUI/Source/CommonInput/Public/CommonInputBaseTypes.h b/Engine/Plugins/Runtime/CommonUI/Source/CommonInput/Public/CommonInputBaseTypes.h
--- a/Engine/Plugins/Runtime/CommonUI/Source/CommonInput/Public/CommonInputBaseTypes.h
+++ b/Engine/Plugins/Runtime/CommonUI/Source/CommonInput/Public/CommonInputBaseTypes.h
@@ -260,6 +260,7 @@ public:
 #if WITH_EDITOR
 	void PostEditChangeProperty(struct FPropertyChangedEvent& PropertyChangedEvent);
 #endif
+	virtual void PostReloadConfig(FProperty* PropertyThatWasLoaded) override;
 
 protected:
 	void InitializeControllerData() const;
