Removed FTickableGameObject from UCommonBoundActionBar. It now sets a ticker when bound actions are updated and when input mode changes (if ticker isn't already set).

#tests Tested in PIE. Opened map screen, toggled between gamepad and keyboard/mouse.

diff --git a/Engine/Plugins/Runtime/CommonUI/Source/CommonUI/Private/Input/CommonBoundActionBar.cpp b/Engine/Plugins/Runtime/CommonUI/Source/CommonUI/Private/Input/CommonBoundActionBar.cpp
--- a/Engine/Plugins/Runtime/CommonUI/Source/CommonUI/Private/Input/CommonBoundActionBar.cpp
+++ b/Engine/Plugins/Runtime/CommonUI/Source/CommonUI/Private/Input/CommonBoundActionBar.cpp
@@ -37,37 +37,6 @@ void UCommonBoundActionBar::SetDisplayOwningPlayerActionsOnly(bool bShouldOnlyDi
 	}
 }
 
-void UCommonBoundActionBar::Tick(float DeltaTime)
-{
-	if (bIsRefreshQueued)
-	{
-		HandleDeferredDisplayUpdate();
-	}
-}
-
-ETickableTickType UCommonBoundActionBar::GetTickableTickType() const
-{
-	return (IsTemplate() ? ETickableTickType::Never : ETickableTickType::Always);
-}
-
-TStatId UCommonBoundActionBar::GetStatId() const
-{
-	RETURN_QUICK_DECLARE_CYCLE_STAT(UCommonBoundActionBar, STATGROUP_Tickables);
-}
-
-bool UCommonBoundActionBar::IsTickableWhenPaused() const
-{
-	return true;
-}
-
-void UCommonBoundActionBar::BeginDestroy()
-{
-	Super::BeginDestroy();
-
-	// Make sure the tick is completely disabled
-	SetTickableTickType(ETickableTickType::Never);
-}
-
 void UCommonBoundActionBar::OnWidgetRebuilt()
 {
 	Super::OnWidgetRebuilt();
@@ -136,7 +105,16 @@ void UCommonBoundActionBar::HandleBoundActionsUpdated(bool bFromOwningPlayer)
 {
 	if (bFromOwningPlayer || !bDisplayOwningPlayerActionsOnly)
 	{
-		bIsRefreshQueued = true;
+		if (!bIsRefreshQueued)
+		{
+			bIsRefreshQueued = true;
+
+			FTSTicker::GetCoreTicker().AddTicker(FTickerDelegate::CreateWeakLambda(this, [this](float InDelta)
+			{
+				HandleDeferredDisplayUpdate();
+				return false;
+			}));
+		}
 	}
 }
 
@@ -360,7 +338,16 @@ void UCommonBoundActionBar::HandlePlayerAdded(int32 PlayerIdx)
 
 void UCommonBoundActionBar::HandledInputTypeUpdated(ECommonInputType InputType)
 {
-	bIsRefreshQueued = true;
+	if (!bIsRefreshQueued)
+	{
+		bIsRefreshQueued = true;
+
+		FTSTicker::GetCoreTicker().AddTicker(FTickerDelegate::CreateWeakLambda(this, [this](float InDelta)
+		{
+			HandleDeferredDisplayUpdate();
+			return false;
+		}));
+	}
 }
 
 void UCommonBoundActionBar::MonitorPlayerActions(const ULocalPlayer* NewPlayer)
diff --git a/Engine/Plugins/Runtime/CommonUI/Source/CommonUI/Public/Input/CommonBoundActionBar.h b/Engine/Plugins/Runtime/CommonUI/Source/CommonUI/Public/Input/CommonBoundActionBar.h
--- a/Engine/Plugins/Runtime/CommonUI/Source/CommonUI/Public/Input/CommonBoundActionBar.h
+++ b/Engine/Plugins/Runtime/CommonUI/Source/CommonUI/Public/Input/CommonBoundActionBar.h
@@ -5,7 +5,6 @@
 #include "CommonButtonBase.h"
 #include "Components/DynamicEntryBoxBase.h"
 #include "CommonInputTypeEnum.h"
-#include "Tickable.h"
 #include "CommonBoundActionBar.generated.h"
 
 class ICommonBoundActionButtonInterface;
@@ -20,7 +19,7 @@ DECLARE_DYNAMIC_MULTICAST_DELEGATE(FActionBarUpdated);
  * A box populated with current actions available per CommonUI's Input Handler.
  */
 UCLASS(Blueprintable, ClassGroup = UI, meta = (Category = "Common UI"))
-class COMMONUI_API UCommonBoundActionBar : public UDynamicEntryBoxBase, public FTickableGameObject
+class COMMONUI_API UCommonBoundActionBar : public UDynamicEntryBoxBase
 {
 	GENERATED_BODY()
 
@@ -28,15 +27,7 @@ public:
 	UFUNCTION(BlueprintCallable, Category = CommonBoundActionBar)
 	void SetDisplayOwningPlayerActionsOnly(bool bShouldOnlyDisplayOwningPlayerActions);
 
-	//~ FTickableGameObject Begin
-	virtual void Tick(float DeltaTime) override;
-	virtual ETickableTickType GetTickableTickType() const override;
-	virtual TStatId GetStatId() const override;
-	virtual bool IsTickableWhenPaused() const override;
-	//~ FTickableGameObject End
-
 protected:
-	virtual void BeginDestroy() override;
 	virtual void OnWidgetRebuilt() override;
 	virtual void SynchronizeProperties() override;
 	virtual void ReleaseSlateResources(bool bReleaseChildren) override;
