Remove ensure on CommonInputSubsystem for functions that can be called by BP functions.

diff --git a/Engine/Plugins/Runtime/CommonUI/Source/CommonUI/Private/CommonActionWidget.cpp b/Engine/Plugins/Runtime/CommonUI/Source/CommonUI/Private/CommonActionWidget.cpp
--- a/Engine/Plugins/Runtime/CommonUI/Source/CommonUI/Private/CommonActionWidget.cpp
+++ b/Engine/Plugins/Runtime/CommonUI/Source/CommonUI/Private/CommonActionWidget.cpp
@@ -183,7 +183,7 @@ FText UCommonActionWidget::GetDisplayText() const
 	}
 
 	const UCommonInputSubsystem* CommonInputSubsystem = GetInputSubsystem();
-	if (GetGameInstance() && ensure(CommonInputSubsystem))
+	if (GetGameInstance() && CommonInputSubsystem)
 	{
 		if (const FCommonInputActionDataBase* InputActionData = GetInputActionData())
 		{
@@ -237,7 +237,7 @@ bool UCommonActionWidget::IsHeldAction() const
 	}
 
 	const UCommonInputSubsystem* CommonInputSubsystem = GetInputSubsystem();
-	if (GetGameInstance() && ensure(CommonInputSubsystem))
+	if (GetGameInstance() && CommonInputSubsystem)
 	{
 		if (const FCommonInputActionDataBase* InputActionData = GetInputActionData())
 		{
