CommonUI: Fix UCommonUIInputSettings::bLinkCursorToGamepadFocus not working

#rnx

diff --git a/Engine/Plugins/Runtime/CommonUI/Source/CommonUI/Private/Input/CommonAnalogCursor.cpp b/Engine/Plugins/Runtime/CommonUI/Source/CommonUI/Private/Input/CommonAnalogCursor.cpp
--- a/Engine/Plugins/Runtime/CommonUI/Source/CommonUI/Private/Input/CommonAnalogCursor.cpp
+++ b/Engine/Plugins/Runtime/CommonUI/Source/CommonUI/Private/Input/CommonAnalogCursor.cpp
@@ -80,7 +80,7 @@ void FCommonAnalogCursor::Tick(const float DeltaTime, FSlateApplication& SlateAp
 			UCommonInputSubsystem& InputSubsystem = ActionRouter.GetInputSubsystem();
 			InputSubsystem.SetCursorPosition(NewPosition, false);
 		}
-		else
+		else if (UCommonUIInputSettings::Get().ShouldLinkCursorToGamepadFocus())
 		{
 			TSharedPtr<SWidget> PinnedLastCursorTarget = LastCursorTarget.Pin();
 
diff --git a/Engine/Plugins/Runtime/CommonUI/Source/CommonUI/Public/Input/CommonUIInputSettings.h b/Engine/Plugins/Runtime/CommonUI/Source/CommonUI/Public/Input/CommonUIInputSettings.h
--- a/Engine/Plugins/Runtime/CommonUI/Source/CommonUI/Public/Input/CommonUIInputSettings.h
+++ b/Engine/Plugins/Runtime/CommonUI/Source/CommonUI/Public/Input/CommonUIInputSettings.h
@@ -118,6 +118,7 @@ public:
 	COMMONUI_API const FUIInputAction* FindAction(FUIActionTag ActionTag) const;
 	COMMONUI_API const TArray<FUIInputAction>& GetUIInputActions() const { return InputActions; }
 	const FCommonAnalogCursorSettings& GetAnalogCursorSettings() const { return AnalogCursorSettings; }
+	bool ShouldLinkCursorToGamepadFocus() const { return bLinkCursorToGamepadFocus; }
 //
 private:
 	/** True to have the mouse pointer automatically moved to the center of whatever widget is currently focused while using a gamepad. */
