CommonUI: Support registering inputs for multiple local players on a single widget

#rnx

diff --git a/Engine/Plugins/Runtime/CommonUI/Source/CommonUI/Private/CommonActionWidget.cpp b/Engine/Plugins/Runtime/CommonUI/Source/CommonUI/Private/CommonActionWidget.cpp
--- a/Engine/Plugins/Runtime/CommonUI/Source/CommonUI/Private/CommonActionWidget.cpp
+++ b/Engine/Plugins/Runtime/CommonUI/Source/CommonUI/Private/CommonActionWidget.cpp
@@ -159,8 +159,8 @@ UCommonInputSubsystem* UCommonActionWidget::GetInputSubsystem() const
 {
 	// In the new system, we may be representing an action for any player, not necessarily the one that technically owns this action icon widget
 	// We want to be sure to use the LocalPlayer that the binding is actually for so we can display the icon that corresponds to their current input method
-	const UWidget* BoundWidget = DisplayedBindingHandle.GetBoundWidget();
-	const ULocalPlayer* BindingOwner = BoundWidget ? BoundWidget->GetOwningLocalPlayer() : GetOwningLocalPlayer();
+	const ULocalPlayer* BoundLocalPlayer = DisplayedBindingHandle.GetBoundLocalPlayer();
+	const ULocalPlayer* BindingOwner = BoundLocalPlayer ? BoundLocalPlayer : GetOwningLocalPlayer();
 	return UCommonInputSubsystem::Get(BindingOwner);
 }
 
@@ -421,9 +421,9 @@ void UCommonActionWidget::ListenToInputMethodChanged(bool bListen)
 	}
 
 	if (CommonUI::IsEnhancedInputSupportEnabled())
-	{
-		const UWidget* BoundWidget = DisplayedBindingHandle.GetBoundWidget();
-		const ULocalPlayer* LocalPlayer = BoundWidget ? BoundWidget->GetOwningLocalPlayer() : GetOwningLocalPlayer();
+	{	
+		const ULocalPlayer* BoundLocalPlayer = DisplayedBindingHandle.GetBoundLocalPlayer();
+     	const ULocalPlayer* LocalPlayer = BoundLocalPlayer ? BoundLocalPlayer : GetOwningLocalPlayer();
 		if (UEnhancedInputLocalPlayerSubsystem* EnhancedInputLocalPlayerSubsystem = ULocalPlayer::GetSubsystem<UEnhancedInputLocalPlayerSubsystem>(LocalPlayer))
 		{
 			if (bListen)
diff --git a/Engine/Plugins/Runtime/CommonUI/Source/CommonUI/Private/Input/CommonBoundActionButton.cpp b/Engine/Plugins/Runtime/CommonUI/Source/CommonUI/Private/Input/CommonBoundActionButton.cpp
--- a/Engine/Plugins/Runtime/CommonUI/Source/CommonUI/Private/Input/CommonBoundActionButton.cpp
+++ b/Engine/Plugins/Runtime/CommonUI/Source/CommonUI/Private/Input/CommonBoundActionButton.cpp
@@ -56,8 +56,7 @@ void UCommonBoundActionButton::UpdateInputActionWidget()
 		FText ActionDisplayName = BindingHandle.GetDisplayName();
 		if (BindingHandle.IsValid())
 		{
-			const UWidget* BoundWidget = BindingHandle.GetBoundWidget();
-			ULocalPlayer* BindingOwner = BoundWidget ? BoundWidget->GetOwningLocalPlayer() : nullptr;
+			ULocalPlayer* BindingOwner = BindingHandle.GetBoundLocalPlayer();
 			if (ensure(BindingOwner) && BindingOwner != GetOwningLocalPlayer())
 			{
 				TOptional<int32> BoundPlayerIndex = FSlateApplication::Get().GetUserIndexForController(BindingOwner->GetControllerId());
diff --git a/Engine/Plugins/Runtime/CommonUI/Source/CommonUI/Private/Input/CommonUIActionRouterBase.cpp b/Engine/Plugins/Runtime/CommonUI/Source/CommonUI/Private/Input/CommonUIActionRouterBase.cpp
--- a/Engine/Plugins/Runtime/CommonUI/Source/CommonUI/Private/Input/CommonUIActionRouterBase.cpp
+++ b/Engine/Plugins/Runtime/CommonUI/Source/CommonUI/Private/Input/CommonUIActionRouterBase.cpp
@@ -48,6 +48,12 @@ static const FAutoConsoleVariableRef CVarResetUIInputConfigOnActivatableTreeDeac
 	bResetUIInputConfigOnActivatableTreeDeactivation,
 	TEXT("Controls if input config is reset when root is changed via deactivation."));
 
+bool bSupportMultiUserInput = true;
+static const FAutoConsoleVariableRef CVarSupportMultiUserInput(
+	TEXT("CommonUI.SupportMultiUserInput"),
+	bSupportMultiUserInput,
+	TEXT("Whether or not action routers can forward inputs to other action routers to support widgets binding inputs for multiple local players"));
+
 //////////////////////////////////////////////////////////////////////////
 
 static bool bTraceInputConfig = false;
@@ -175,7 +181,7 @@ UCommonUIActionRouterBase::UCommonUIActionRouterBase()
 
 FUIActionBindingHandle UCommonUIActionRouterBase::RegisterUIActionBinding(const UWidget& Widget, const FBindUIActionArgs& BindActionArgs)
 {
-	FUIActionBindingHandle BindingHandle = FUIActionBinding::TryCreate(Widget, BindActionArgs);
+	FUIActionBindingHandle BindingHandle = FUIActionBinding::TryCreate(Widget, BindActionArgs, GetLocalPlayerIndex());
 	if (BindingHandle.IsValid())
 	{
 		FActivatableTreeNodePtr OwnerNode = nullptr;
@@ -466,8 +472,6 @@ ERouteUIInputResult UCommonUIActionRouterBase::ProcessInput(FKey Key, EInputEven
 	}
 #endif
 
-	ECommonInputMode ActiveMode = GetActiveInputMode();
-
 	// Also check for repeat event here as if input is flushed when a key is being held, we will receive a released event and then continue to receive repeat events without a pressed event
 	if (InputEvent == EInputEvent::IE_Pressed || InputEvent == EInputEvent::IE_Repeat)
 	{
@@ -478,60 +482,100 @@ ERouteUIInputResult UCommonUIActionRouterBase::ProcessInput(FKey Key, EInputEven
 		HeldKeys.RemoveSwap(Key);
 	}
 
+	const ECommonInputMode ActiveMode = GetActiveInputMode();
+	const int32 OwningUserIndex = GetLocalPlayerIndex();
+
 	// Begin with a pass to see if the input corresponds to a hold action
 	// We do this first to make sure that a higher-priority press binding doesn't prevent a hold on the same key from being triggerable
-	EProcessHoldActionResult ProcessHoldResult = PersistentActions->ProcessHoldInput(ActiveMode, Key, InputEvent);
-	
-	if (ProcessHoldResult == EProcessHoldActionResult::Unhandled)
+	const auto ProcessHoldInputFunc = [ActiveMode, Key, InputEvent, OwningUserIndex](const UCommonUIActionRouterBase& ActionRouter)
 	{
-		if (bIsActivatableTreeEnabled && ActiveRootNode)
-		{
-			ProcessHoldResult = ActiveRootNode->ProcessHoldInput(ActiveMode, Key, InputEvent);
-		}
-
+		EProcessHoldActionResult ProcessHoldResult = ActionRouter.PersistentActions->ProcessHoldInput(ActiveMode, Key, InputEvent, OwningUserIndex);
+	
 		if (ProcessHoldResult == EProcessHoldActionResult::Unhandled)
 		{
-			ProcessHoldResult = ProcessHoldInputOnActionDomains(ActiveMode, Key, InputEvent);
+			if (ActionRouter.bIsActivatableTreeEnabled && ActionRouter.ActiveRootNode)
+			{
+				ProcessHoldResult = ActionRouter.ActiveRootNode->ProcessHoldInput(ActiveMode, Key, InputEvent, OwningUserIndex);
+			}
+
+			if (ProcessHoldResult == EProcessHoldActionResult::Unhandled)
+			{
+				ProcessHoldResult = ActionRouter.ProcessHoldInputOnActionDomains(ActiveMode, Key, InputEvent, OwningUserIndex);
+			}
 		}
-	}
 
-	const auto ProcessNormalInputFunc = [Key, ActiveMode, this](EInputEvent Event)
+		return ProcessHoldResult;
+	};
+
+	const auto ProcessNormalInputFunc = [Key, ActiveMode, OwningUserIndex](const UCommonUIActionRouterBase& ActionRouter, EInputEvent Event)
+	{
+		bool bHandled = ActionRouter.PersistentActions->ProcessNormalInput(ActiveMode, Key, Event, OwningUserIndex);
+
+		if (!bHandled)
 		{
-			bool bHandled = PersistentActions->ProcessNormalInput(ActiveMode, Key, Event);
+			if (ActionRouter.bIsActivatableTreeEnabled && ActionRouter.ActiveRootNode)
+			{
+				bHandled = ActionRouter.ActiveRootNode->ProcessNormalInput(ActiveMode, Key, Event, OwningUserIndex);
+			}
 
 			if (!bHandled)
 			{
-				if (bIsActivatableTreeEnabled && ActiveRootNode)
-				{
-					bHandled = ActiveRootNode->ProcessNormalInput(ActiveMode, Key, Event);
-				}
-
-				if (!bHandled)
-				{
-					bHandled = ProcessInputOnActionDomains(ActiveMode, Key, Event);
-				}
+				bHandled = ActionRouter.ProcessInputOnActionDomains(ActiveMode, Key, Event, OwningUserIndex);
 			}
+		}
 
-			return bHandled;
-		};
+		return bHandled;
+	};
 
-	bool bHandledInput = ProcessHoldResult == EProcessHoldActionResult::Handled;
-	if (!bHandledInput)
+	const auto ProcessInputOnActionRouter = [&ProcessHoldInputFunc, &ProcessNormalInputFunc, InputEvent](const UCommonUIActionRouterBase& ActionRouter)
 	{
+		EProcessHoldActionResult ProcessHoldResult = ProcessHoldInputFunc(ActionRouter);
+		if (ProcessHoldResult == EProcessHoldActionResult::Handled)
+		{
+			return true;
+		}
+		
 		if (ProcessHoldResult == EProcessHoldActionResult::GeneratePress)
 		{
 			// A hold action was in progress but quickly aborted, so we want to generate a press action now for any normal bindings that are interested
-			ProcessNormalInputFunc(IE_Pressed);
+			ProcessNormalInputFunc(ActionRouter, IE_Pressed);
 		}
 
 		// Even if no widget cares about this input, we don't want to let anything through to the actual game while we're in menu mode
-		bHandledInput = ProcessNormalInputFunc(InputEvent); 
+		return ProcessNormalInputFunc(ActionRouter, InputEvent);
+	};
+
+	bool bHandledInput = ProcessInputOnActionRouter(*this);
+	if (bSupportMultiUserInput && !bHandledInput)
+	{
+		const ULocalPlayer* LocalPlayer = GetLocalPlayerChecked();
+		if (const UGameInstance* GameInstance = LocalPlayer->GetGameInstance())
+		{
+			for (const ULocalPlayer* OtherPlayer : GameInstance->GetLocalPlayers())
+			{
+				if (OtherPlayer == LocalPlayer)
+				{
+					continue;
+				}
+
+				// If necessary, this could be sped up by caching something to indicate which action routers have bindings for which players
+				if (const UCommonUIActionRouterBase* OtherActionRouter = ULocalPlayer::GetSubsystem<UCommonUIActionRouterBase>(OtherPlayer))
+				{
+					if (ProcessInputOnActionRouter(*OtherActionRouter))
+					{
+						bHandledInput = true;
+						break;
+					}
+				}
+			}
+		}
 	}
 
 	if (bHandledInput)
 	{
 		return ERouteUIInputResult::Handled;
 	}
+
 	return CanProcessNormalGameInput() ? ERouteUIInputResult::Unhandled : ERouteUIInputResult::BlockGameInput;
 }
 
@@ -545,12 +589,53 @@ UCommonInputSubsystem& UCommonUIActionRouterBase::GetInputSubsystem() const
 void UCommonUIActionRouterBase::FlushInput()
 {
 	const ECommonInputMode ActiveMode = GetActiveInputMode();
+	const int32 OwningUserIndex = GetLocalPlayerIndex();
+	const auto FlushInputOnActionRouter = [ActiveMode, OwningUserIndex](const UCommonUIActionRouterBase& ActionRouter, const FKey& HeldKey)
+	{
+		EProcessHoldActionResult ProcessHoldResult = ActionRouter.PersistentActions->ProcessHoldInput(ActiveMode, HeldKey, IE_Released, OwningUserIndex);
+		if (ActionRouter.bIsActivatableTreeEnabled && ProcessHoldResult == EProcessHoldActionResult::Unhandled)
+		{
+			if (ActionRouter.ActiveRootNode)
+			{
+				ProcessHoldResult = ActionRouter.ActiveRootNode->ProcessHoldInput(ActiveMode, HeldKey, IE_Released, OwningUserIndex);
+			}
+
+			if (ProcessHoldResult == EProcessHoldActionResult::Unhandled)
+			{
+				ProcessHoldResult = ActionRouter.ProcessHoldInputOnActionDomains(ActiveMode, HeldKey, IE_Released, OwningUserIndex);
+			}
+		}
+
+		return ProcessHoldResult;
+	};
+
 	for (const FKey& HeldKey : HeldKeys)
 	{
-		const EProcessHoldActionResult ProcessHoldResult = PersistentActions->ProcessHoldInput(ActiveMode, HeldKey, EInputEvent::IE_Released);
-		if (bIsActivatableTreeEnabled && ActiveRootNode && ProcessHoldResult == EProcessHoldActionResult::Unhandled)
+		EProcessHoldActionResult ProcessHoldResult = FlushInputOnActionRouter(*this, HeldKey);
+		
+		if (bSupportMultiUserInput && ProcessHoldResult == EProcessHoldActionResult::Unhandled)
 		{
-			ActiveRootNode->ProcessHoldInput(ActiveMode, HeldKey, EInputEvent::IE_Released);
+			ULocalPlayer* LocalPlayer = GetLocalPlayerChecked();
+			if (UGameInstance* GameInstance = LocalPlayer->GetGameInstance())
+			{
+				for (ULocalPlayer* OtherPlayer : GameInstance->GetLocalPlayers())
+				{
+					if (OtherPlayer == LocalPlayer)
+					{
+						continue;
+					}
+
+					// If necessary, this could be sped up by caching something to indicate which action routers have widgets with bindings for which players
+					if (const UCommonUIActionRouterBase* OtherActionRouter = ULocalPlayer::GetSubsystem<UCommonUIActionRouterBase>(OtherPlayer))
+					{
+						ProcessHoldResult = FlushInputOnActionRouter(*OtherActionRouter, HeldKey);
+						if (ProcessHoldResult != EProcessHoldActionResult::Unhandled)
+						{
+							break;
+						}
+					}
+				}
+			}
 		}
 	}
 
@@ -855,7 +940,7 @@ const UCommonInputActionDomainTable* UCommonUIActionRouterBase::GetActionDomainT
 	return InputSubsytem ? InputSubsytem->GetActionDomainTable() : nullptr;
 }
 
-bool UCommonUIActionRouterBase::ProcessInputOnActionDomains(ECommonInputMode ActiveInputMode, FKey Key, EInputEvent InputEvent) const
+bool UCommonUIActionRouterBase::ProcessInputOnActionDomains(ECommonInputMode ActiveInputMode, FKey Key, EInputEvent InputEvent, int32 UserIndex) const
 {
 	const UCommonInputActionDomainTable* ActionDomainTable = GetActionDomainTable();
 	if (!ActionDomainTable)
@@ -881,7 +966,7 @@ bool UCommonUIActionRouterBase::ProcessInputOnActionDomains(ECommonInputMode Act
 		{
 			if (RootNode->IsWidgetActivated())
 			{
-				bool bInputEventHandled = RootNode->ProcessNormalInput(ActiveInputMode, Key, InputEvent);
+				bool bInputEventHandled = RootNode->ProcessNormalInput(ActiveInputMode, Key, InputEvent, UserIndex);
 				bInputEventHandledInDomain |= bInputEventHandled;
 				bDomainHadActiveRoots = true;
 
@@ -903,7 +988,7 @@ bool UCommonUIActionRouterBase::ProcessInputOnActionDomains(ECommonInputMode Act
 	return bInputEventHandledAtLeastOnce;
 }
 
-EProcessHoldActionResult UCommonUIActionRouterBase::ProcessHoldInputOnActionDomains(ECommonInputMode ActiveInputMode, FKey Key, EInputEvent InputEvent) const
+EProcessHoldActionResult UCommonUIActionRouterBase::ProcessHoldInputOnActionDomains(ECommonInputMode ActiveInputMode, FKey Key, EInputEvent InputEvent, int32 UserIndex) const
 {
 	EProcessHoldActionResult HoldActionResult = EProcessHoldActionResult::Unhandled;
 	
@@ -929,7 +1014,7 @@ EProcessHoldActionResult UCommonUIActionRouterBase::ProcessHoldInputOnActionDoma
 		{
 			if (RootNode->IsReceivingInput() && HoldActionResult == EProcessHoldActionResult::Unhandled)
 			{
-				HoldActionResult = RootNode->ProcessHoldInput(ActiveInputMode, Key, InputEvent);
+				HoldActionResult = RootNode->ProcessHoldInput(ActiveInputMode, Key, InputEvent, UserIndex);
 				bInputEventHandledInDomain |= HoldActionResult == EProcessHoldActionResult::Handled;
 				bDomainHadActiveRoots = true;
 
@@ -1157,16 +1242,52 @@ bool UCommonUIActionRouterBase::Tick(float DeltaTime)
 
 		ActiveActionDomainRootsPendingPaint.Append(ActiveActionDomainRootsPendingPaintCopy);
 	}
-
+	
 	const ECommonInputMode ActiveMode = GetActiveInputMode();
+	const int32 OwningUserIndex = GetLocalPlayerIndex();
+	const auto TickInputOnActionRouter = [ActiveMode, OwningUserIndex](const UCommonUIActionRouterBase& ActionRouter, const FKey& HeldKey)
+	{
+		EProcessHoldActionResult ProcessHoldResult = ActionRouter.PersistentActions->ProcessHoldInput(ActiveMode, HeldKey, IE_Repeat, OwningUserIndex);
+		if (ActionRouter.bIsActivatableTreeEnabled && ProcessHoldResult == EProcessHoldActionResult::Unhandled)
+		{
+			if (ActionRouter.ActiveRootNode)
+			{
+				ProcessHoldResult = ActionRouter.ActiveRootNode->ProcessHoldInput(ActiveMode, HeldKey, IE_Repeat, OwningUserIndex);
+			}
+
+			if (ProcessHoldResult == EProcessHoldActionResult::Unhandled)
+			{
+				ProcessHoldResult = ActionRouter.ProcessHoldInputOnActionDomains(ActiveMode, HeldKey, IE_Repeat, OwningUserIndex);
+			}
+		}
+		return ProcessHoldResult;
+	};
+	
 	for (const FKey& HeldKey : HeldKeys)
 	{
-		const EProcessHoldActionResult ProcessHoldResult = PersistentActions->ProcessHoldInput(ActiveMode, HeldKey, EInputEvent::IE_Repeat);
-		if (bIsActivatableTreeEnabled && ActiveRootNode && ProcessHoldResult == EProcessHoldActionResult::Unhandled)
+		EProcessHoldActionResult ProcessHoldResult = TickInputOnActionRouter(*this, HeldKey);
+		if (bSupportMultiUserInput && ProcessHoldResult == EProcessHoldActionResult::Unhandled)
 		{
-			if (ActiveRootNode->ProcessHoldInput(ActiveMode, HeldKey, EInputEvent::IE_Repeat) == EProcessHoldActionResult::Unhandled)
+			ULocalPlayer* LocalPlayer = GetLocalPlayerChecked();
+			if (UGameInstance* GameInstance = LocalPlayer->GetGameInstance())
 			{
-				ProcessHoldInputOnActionDomains(ActiveMode, HeldKey, EInputEvent::IE_Repeat);
+				for (ULocalPlayer* OtherPlayer : GameInstance->GetLocalPlayers())
+				{
+					if (OtherPlayer == LocalPlayer)
+					{
+						continue;
+					}
+
+					// If necessary, this could be sped up by caching something to indicate which action routers have widgets with bindings for which players
+					if (const UCommonUIActionRouterBase* OtherActionRouter = ULocalPlayer::GetSubsystem<UCommonUIActionRouterBase>(OtherPlayer))
+					{
+						ProcessHoldResult = TickInputOnActionRouter(*OtherActionRouter, HeldKey);
+						if (ProcessHoldResult != EProcessHoldActionResult::Unhandled)
+						{
+							break;
+						}
+					}
+				}
 			}
 		}
 	}
diff --git a/Engine/Plugins/Runtime/CommonUI/Source/CommonUI/Private/Input/UIActionRouterTypes.cpp b/Engine/Plugins/Runtime/CommonUI/Source/CommonUI/Private/Input/UIActionRouterTypes.cpp
--- a/Engine/Plugins/Runtime/CommonUI/Source/CommonUI/Private/Input/UIActionRouterTypes.cpp
+++ b/Engine/Plugins/Runtime/CommonUI/Source/CommonUI/Private/Input/UIActionRouterTypes.cpp
@@ -18,6 +18,7 @@
 #include "Widgets/SViewport.h"
 #include "CommonButtonBase.h"
 #include "CommonUIEditorSettings.h"
+#include "Engine/GameInstance.h"
 #include "ICommonUIModule.h"
 
 DEFINE_LOG_CATEGORY(LogUIActionRouter);
@@ -211,6 +212,13 @@ FUIActionBinding::FUIActionBinding(const UWidget& InBoundWidget, const FBindUIAc
 }
 
 FUIActionBindingHandle FUIActionBinding::TryCreate(const UWidget& InBoundWidget, const FBindUIActionArgs& BindArgs)
+{
+	const ULocalPlayer* OwningPlayer = InBoundWidget.GetOwningLocalPlayer();
+	const int32 UserIndex = OwningPlayer ? OwningPlayer->GetLocalPlayerIndex() : INDEX_NONE;
+	return TryCreate(InBoundWidget, BindArgs, UserIndex);
+}
+
+FUIActionBindingHandle FUIActionBinding::TryCreate(const UWidget& InBoundWidget, const FBindUIActionArgs& BindArgs, int32 UserIndex)
 {
 	bool bIsEnhancedInputSupportEnabled = CommonUI::IsEnhancedInputSupportEnabled();
 	if (BindArgs.GetActionName().IsNone() && (!bIsEnhancedInputSupportEnabled || !BindArgs.InputAction.IsValid()))
@@ -238,21 +246,28 @@ FUIActionBindingHandle FUIActionBinding::TryCreate(const UWidget& InBoundWidget,
 		UE_LOG(LogUIActionRouter, Error, TEXT("Cannot bind widget [%s] to action [%s] - provided input action is invalid."), *InBoundWidget.GetName(), *BindArgs.GetActionName().ToString());
 		return FUIActionBindingHandle();
 	}
+	else if (UserIndex == INDEX_NONE)
+	{
+		UE_LOG(LogUIActionRouter, Error, TEXT("Cannot bind widget [%s] to action [%s] [%s] - invalid UserIndex."), *BindArgs.GetActionName().ToString(), InputEventToString(BindArgs.KeyEvent), *InBoundWidget.GetName());
+		return FUIActionBindingHandle();
+	}
 	
 	// Make sure there is no existing binding for the same action associated with the same widget
 	for (const TPair<FUIActionBindingHandle, TSharedPtr<FUIActionBinding>>& HandleBindingPair : AllRegistrationsByHandle)
 	{
 		if (HandleBindingPair.Value->ActionName == BindArgs.GetActionName() &&
 			HandleBindingPair.Value->InputEvent == BindArgs.KeyEvent &&
-			HandleBindingPair.Value->BoundWidget.Get() == &InBoundWidget)
+			HandleBindingPair.Value->BoundWidget.Get() == &InBoundWidget &&
+			HandleBindingPair.Value->UserIndex == UserIndex)
 		{
-			UE_LOG(LogUIActionRouter, Error, TEXT("Widget [%s] is already bound to action [%s] [%s]! A given widget can only bind the same action once. Unregister existing binding first if you wish to change any aspect of the binding"), 
-				*InBoundWidget.GetName(), *HandleBindingPair.Value->ActionName.ToString(), InputEventToString(BindArgs.KeyEvent));
+			UE_LOG(LogUIActionRouter, Error, TEXT("Widget [%s] is already bound to action [%s] [%s] for user [%d]! A given widget can only bind the same action once per user. Unregister existing binding first if you wish to change any aspect of the binding"), 
+				*InBoundWidget.GetName(), *HandleBindingPair.Value->ActionName.ToString(), InputEventToString(BindArgs.KeyEvent), UserIndex);
 			return FUIActionBindingHandle();
 		}
 	}
 	
 	TSharedPtr<FUIActionBinding> NewRegistration = MakeShareable(new FUIActionBinding(InBoundWidget, BindArgs));
+	NewRegistration->UserIndex = UserIndex;
 	FUIActionBindingHandle Handle = NewRegistration->Handle;
 	AllRegistrationsByHandle.Add(Handle, MoveTemp(NewRegistration));
 
@@ -570,6 +585,26 @@ const UWidget* FUIActionBindingHandle::GetBoundWidget() const
 	return nullptr;
 }
 
+ULocalPlayer* FUIActionBindingHandle::GetBoundLocalPlayer() const
+{
+	if (TSharedPtr<const FUIActionBinding> Binding = FUIActionBinding::FindBinding(*this))
+	{
+		if (const UWidget* BoundWidget = Binding->BoundWidget.Get())
+		{
+			const UGameInstance* GameInstance = BoundWidget->GetGameInstance();
+			if (GameInstance && Binding->UserIndex != INDEX_NONE)
+			{
+				return GameInstance->GetLocalPlayerByIndex(Binding->UserIndex);
+			}
+			else
+			{
+				return BoundWidget->GetOwningLocalPlayer();
+			}
+		}
+	}
+	return nullptr;
+}
+
 
 //////////////////////////////////////////////////////////////////////////
 // FUIInputConfig
@@ -626,13 +661,13 @@ FActionRouterBindingCollection::FActionRouterBindingCollection(UCommonUIActionRo
 	: ActionRouterPtr(&OwningRouter)
 {}
 
-EProcessHoldActionResult FActionRouterBindingCollection::ProcessHoldInput(ECommonInputMode ActiveInputMode, FKey Key, EInputEvent InputEvent) const
+EProcessHoldActionResult FActionRouterBindingCollection::ProcessHoldInput(ECommonInputMode ActiveInputMode, FKey Key, EInputEvent InputEvent, int32 UserIndex) const
 {
 	for (FUIActionBindingHandle BindingHandle : ActionBindings)
 	{
 		if (TSharedPtr<FUIActionBinding> Binding = FUIActionBinding::FindBinding(BindingHandle))
 		{
-			if (ActiveInputMode == ECommonInputMode::All || ActiveInputMode == Binding->InputMode)
+			if (Binding->UserIndex == UserIndex && (ActiveInputMode == ECommonInputMode::All || ActiveInputMode == Binding->InputMode))
 			{
 				for (const FUIActionKeyMapping& HoldMapping : Binding->HoldMappings)
 				{
@@ -699,13 +734,13 @@ EProcessHoldActionResult FActionRouterBindingCollection::ProcessHoldInput(ECommo
 	return EProcessHoldActionResult::Unhandled;
 }
 
-bool FActionRouterBindingCollection::ProcessNormalInput(ECommonInputMode ActiveInputMode, FKey Key, EInputEvent InputEvent) const
+bool FActionRouterBindingCollection::ProcessNormalInput(ECommonInputMode ActiveInputMode, FKey Key, EInputEvent InputEvent, int32 UserIndex) const
 {
 	for (FUIActionBindingHandle BindingHandle : ActionBindings)
 	{
 		if (TSharedPtr<FUIActionBinding> Binding = FUIActionBinding::FindBinding(BindingHandle))
 		{
-			if (ActiveInputMode == ECommonInputMode::All || ActiveInputMode == Binding->InputMode)
+			if (Binding->UserIndex == UserIndex && (ActiveInputMode == ECommonInputMode::All || ActiveInputMode == Binding->InputMode))
 			{
 				auto TryConsumeInput = [&](const FKey& InKey, const UInputAction* InInputAction)
 				{
@@ -752,7 +787,7 @@ bool FActionRouterBindingCollection::ProcessNormalInput(ECommonInputMode ActiveI
 					if (const UInputAction* InputAction = Binding->InputAction.Get())
 					{
 						TArray<FKey> InputActionKeys;
-						CommonUI::GetEnhancedInputActionKeys(GetActionRouter().GetLocalPlayerChecked(), InputAction, InputActionKeys);
+						CommonUI::GetEnhancedInputActionKeys(BindingHandle.GetBoundLocalPlayer(), InputAction, InputActionKeys);
 						for (const FKey& InputActionKey : InputActionKeys)
 						{
 							if (TryConsumeInput(InputActionKey, InputAction))
@@ -941,35 +976,35 @@ FActivatableTreeNode::~FActivatableTreeNode()
 	}
 }
 
-EProcessHoldActionResult FActivatableTreeNode::ProcessHoldInput(ECommonInputMode ActiveInputMode, FKey Key, EInputEvent InputEvent) const
+EProcessHoldActionResult FActivatableTreeNode::ProcessHoldInput(ECommonInputMode ActiveInputMode, FKey Key, EInputEvent InputEvent, int32 UserIndex) const
 {
 	if (IsReceivingInput())
 	{
 		for (const FActivatableTreeNodeRef& ChildNode : Children)
 		{
-			EProcessHoldActionResult ChildResult = ChildNode->ProcessHoldInput(ActiveInputMode, Key, InputEvent);
+			EProcessHoldActionResult ChildResult = ChildNode->ProcessHoldInput(ActiveInputMode, Key, InputEvent, UserIndex);
 			if (ChildResult != EProcessHoldActionResult::Unhandled)
 			{
 				return ChildResult;
 			}
 		}
-		return FActionRouterBindingCollection::ProcessHoldInput(ActiveInputMode, Key, InputEvent);
+		return FActionRouterBindingCollection::ProcessHoldInput(ActiveInputMode, Key, InputEvent, UserIndex);
 	}
 	return EProcessHoldActionResult::Unhandled;
 }
 
-bool FActivatableTreeNode::ProcessNormalInput(ECommonInputMode ActiveInputMode, FKey Key, EInputEvent InputEvent) const
+bool FActivatableTreeNode::ProcessNormalInput(ECommonInputMode ActiveInputMode, FKey Key, EInputEvent InputEvent, int32 UserIndex) const
 {
 	if (IsReceivingInput())
 	{
 		for (const FActivatableTreeNodeRef& ChildNode : Children)
 		{
-			if (ChildNode->ProcessNormalInput(ActiveInputMode, Key, InputEvent))
+			if (ChildNode->ProcessNormalInput(ActiveInputMode, Key, InputEvent, UserIndex))
 			{
 				return true;
 			}
 		}
-		return FActionRouterBindingCollection::ProcessNormalInput(ActiveInputMode, Key, InputEvent);
+		return FActionRouterBindingCollection::ProcessNormalInput(ActiveInputMode, Key, InputEvent, UserIndex);
 	}
 	return false;
 }
diff --git a/Engine/Plugins/Runtime/CommonUI/Source/CommonUI/Private/Input/UIActionRouterTypes.h b/Engine/Plugins/Runtime/CommonUI/Source/CommonUI/Private/Input/UIActionRouterTypes.h
--- a/Engine/Plugins/Runtime/CommonUI/Source/CommonUI/Private/Input/UIActionRouterTypes.h
+++ b/Engine/Plugins/Runtime/CommonUI/Source/CommonUI/Private/Input/UIActionRouterTypes.h
@@ -48,8 +48,8 @@ class COMMONUI_API FActionRouterBindingCollection : public TSharedFromThis<FActi
 public:
 	virtual ~FActionRouterBindingCollection() {}
 
-	virtual EProcessHoldActionResult ProcessHoldInput(ECommonInputMode ActiveInputMode, FKey Key, EInputEvent InputEvent) const;
-	virtual bool ProcessNormalInput(ECommonInputMode ActiveInputMode, FKey Key, EInputEvent InputEvent) const;
+	virtual EProcessHoldActionResult ProcessHoldInput(ECommonInputMode ActiveInputMode, FKey Key, EInputEvent InputEvent, int32 UserIndex) const;
+	virtual bool ProcessNormalInput(ECommonInputMode ActiveInputMode, FKey Key, EInputEvent InputEvent, int32 UserIndex) const;
 	virtual bool IsReceivingInput() const { return true; }
 
 	void AddBinding(FUIActionBinding& Binding);
@@ -96,8 +96,8 @@ class COMMONUI_API FActivatableTreeNode : public FActionRouterBindingCollection
 public:
 	virtual ~FActivatableTreeNode();
 	
-	virtual EProcessHoldActionResult ProcessHoldInput(ECommonInputMode ActiveInputMode, FKey Key, EInputEvent InputEvent) const override;
-	virtual bool ProcessNormalInput(ECommonInputMode ActiveInputMode, FKey Key, EInputEvent InputEvent) const override;
+	virtual EProcessHoldActionResult ProcessHoldInput(ECommonInputMode ActiveInputMode, FKey Key, EInputEvent InputEvent, int32 UserIndex) const override;
+	virtual bool ProcessNormalInput(ECommonInputMode ActiveInputMode, FKey Key, EInputEvent InputEvent, int32 UserIndex) const override;
 	virtual bool IsReceivingInput() const override { return bCanReceiveInput && IsWidgetActivated(); }
 
 	bool IsWidgetValid() const;
diff --git a/Engine/Plugins/Runtime/CommonUI/Source/CommonUI/Public/Input/CommonUIActionRouterBase.h b/Engine/Plugins/Runtime/CommonUI/Source/CommonUI/Public/Input/CommonUIActionRouterBase.h
--- a/Engine/Plugins/Runtime/CommonUI/Source/CommonUI/Public/Input/CommonUIActionRouterBase.h
+++ b/Engine/Plugins/Runtime/CommonUI/Source/CommonUI/Public/Input/CommonUIActionRouterBase.h
@@ -194,8 +194,8 @@ private:
 	void HandlePostGarbageCollect();
 
 	const UCommonInputActionDomainTable* GetActionDomainTable() const;
-	bool ProcessInputOnActionDomains(ECommonInputMode ActiveInputMode, FKey Key, EInputEvent InputEvent) const;
-	EProcessHoldActionResult ProcessHoldInputOnActionDomains(ECommonInputMode ActiveInputMode, FKey Key, EInputEvent InputEvent) const;
+	bool ProcessInputOnActionDomains(ECommonInputMode ActiveInputMode, FKey Key, EInputEvent InputEvent, int32 UserIndex) const;
+	EProcessHoldActionResult ProcessHoldInputOnActionDomains(ECommonInputMode ActiveInputMode, FKey Key, EInputEvent InputEvent, int32 UserIndex) const;
 	FGameplayTagContainer GetGameplayTagsForInputMode(const ECommonInputMode Mode) const;
 
 	struct FPendingWidgetRegistration
diff --git a/Engine/Plugins/Runtime/CommonUI/Source/CommonUI/Public/Input/UIActionBinding.h b/Engine/Plugins/Runtime/CommonUI/Source/CommonUI/Public/Input/UIActionBinding.h
--- a/Engine/Plugins/Runtime/CommonUI/Source/CommonUI/Public/Input/UIActionBinding.h
+++ b/Engine/Plugins/Runtime/CommonUI/Source/CommonUI/Public/Input/UIActionBinding.h
@@ -29,7 +29,9 @@ struct COMMONUI_API FUIActionBinding
 	FUIActionBinding(const FUIActionBinding&) = delete;
 	FUIActionBinding(FUIActionBinding&&) = delete;
 
+	UE_DEPRECATED(5.6, "Calling FUIActionBinding::TryCreate without a user index is deprecated. Call the version with the UserIndex param.")
 	static FUIActionBindingHandle TryCreate(const UWidget& InBoundWidget, const FBindUIActionArgs& BindArgs);
+	static FUIActionBindingHandle TryCreate(const UWidget& InBoundWidget, const FBindUIActionArgs& BindArgs, int32 UserIndex);
 	
 	static TSharedPtr<FUIActionBinding> FindBinding(FUIActionBindingHandle Handle);
 	static void CleanRegistrations();
@@ -61,6 +63,8 @@ struct COMMONUI_API FUIActionBinding
 	TWeakObjectPtr<const UWidget> BoundWidget;
 	ECommonInputMode InputMode;
 
+	int32 UserIndex = INDEX_NONE;
+
 	bool bDisplayInActionBar = false;
 	FText ActionDisplayName;
 	
diff --git a/Engine/Plugins/Runtime/CommonUI/Source/CommonUI/Public/Input/UIActionBindingHandle.h b/Engine/Plugins/Runtime/CommonUI/Source/CommonUI/Public/Input/UIActionBindingHandle.h
--- a/Engine/Plugins/Runtime/CommonUI/Source/CommonUI/Public/Input/UIActionBindingHandle.h
+++ b/Engine/Plugins/Runtime/CommonUI/Source/CommonUI/Public/Input/UIActionBindingHandle.h
@@ -9,7 +9,7 @@
 enum class EMouseCaptureMode : uint8;
 enum class EMouseLockMode : uint8;
 struct FScriptContainerElement;
-
+class ULocalPlayer;
 class UWidget;
 enum class ECommonInputMode : uint8;
 
@@ -38,6 +38,8 @@ public:
 	void SetDisplayInActionBar(const bool bDisplayInActionBar);
 
 	const UWidget* GetBoundWidget() const;
+	
+	ULocalPlayer* GetBoundLocalPlayer() const;
 
 	FUIActionBindingHandle() {}
 	bool operator==(const FUIActionBindingHandle& Other) const { return RegistrationId == Other.RegistrationId; }
