CommonUI: Fix issue where the focused widget would not get hovered by the virtual cursor when focus moved to the viewport and back to the original widget

#rnx

diff --git a/Engine/Plugins/Runtime/CommonUI/Source/CommonUI/Private/Input/CommonAnalogCursor.cpp b/Engine/Plugins/Runtime/CommonUI/Source/CommonUI/Private/Input/CommonAnalogCursor.cpp
--- a/Engine/Plugins/Runtime/CommonUI/Source/CommonUI/Private/Input/CommonAnalogCursor.cpp
+++ b/Engine/Plugins/Runtime/CommonUI/Source/CommonUI/Private/Input/CommonAnalogCursor.cpp
@@ -70,11 +70,13 @@ void FCommonAnalogCursor::Tick(const float DeltaTime, FSlateApplication& SlateAp
 		RefreshCursorSettings();
 		if (!IsViewportWindowInFocusPath(ActionRouter))
 		{
+			LastCursorTarget.Reset();
 			return;
 		}
 #endif
 		if (bIsAnalogMovementEnabled)
 		{
+			LastCursorTarget.Reset();
 			const FVector2D NewPosition = CalculateTickedCursorPosition(DeltaTime, SlateApp, SlateUser);
 
 			UCommonInputSubsystem& InputSubsystem = ActionRouter.GetInputSubsystem();
@@ -208,6 +210,11 @@ void FCommonAnalogCursor::Tick(const float DeltaTime, FSlateApplication& SlateAp
 			}
 		}
 	}
+	else
+	{
+		// Since we're not processing cursor target this frame, the cursor position may change externally and therefore invalidate our cache
+		LastCursorTarget.Reset();
+	}
 }
 
 bool FCommonAnalogCursor::HandleKeyDownEvent(FSlateApplication& SlateApp, const FKeyEvent& InKeyEvent)
