CommonUI: Downgrade checks to ensures in FActivatableTreeRoot::FocusLeafmostNode

#rnx

diff --git a/Engine/Plugins/Runtime/CommonUI/Source/CommonUI/Private/Input/UIActionRouterTypes.cpp b/Engine/Plugins/Runtime/CommonUI/Source/CommonUI/Private/Input/UIActionRouterTypes.cpp
--- a/Engine/Plugins/Runtime/CommonUI/Source/CommonUI/Private/Input/UIActionRouterTypes.cpp
+++ b/Engine/Plugins/Runtime/CommonUI/Source/CommonUI/Private/Input/UIActionRouterTypes.cpp
@@ -1585,24 +1585,21 @@ void FActivatableTreeRoot::FocusLeafmostNode()
 		return;
 	}
 
-	check(LeafmostActiveNode.IsValid());
-	if (!LeafmostActiveNode.IsValid())
+	if (!ensure(LeafmostActiveNode.IsValid()))
 	{
 		UE_LOG(LogUIActionRouter, Error, TEXT("Cannot focus leaf most node - invalid LeafmostActiveNode"));
 		return;
 	}
 	FActivatableTreeNodePtr PinnedLeafmostNode = LeafmostActiveNode.Pin();
 
-	check(PinnedLeafmostNode);
-	if (!PinnedLeafmostNode)
+	if (!ensure(PinnedLeafmostNode))
 	{
 		UE_LOG(LogUIActionRouter, Error, TEXT("Cannot focus leaf most node - invalid PinnedLeafmostNode"));
 		return;
 	}
 
 	UCommonActivatableWidget* LeafWidget = PinnedLeafmostNode->GetWidget();
-	check(LeafWidget);
-	if (!LeafWidget)
+	if (!ensure(LeafWidget))
 	{
 		UE_LOG(LogUIActionRouter, Error, TEXT("Cannot focus leaf most node - invalid LeafWidget"));
 		return;
