CommonUI: Fix UCommonUIActionRouterBase::HandleSlateFocusChanging not considering leaf nodes of the active root node nor action domain nodes

- Additionally fixed logic in FActivatableTreeNode::CacheFocusRestorationTarget() that could lead to caching a widget that lived outside of the node's hierarchy
- Also replace checks for GetType() == "SObjectWidget" with checks for FCommonActivatableSlateMetaData so SObjectWidget subclasses (such as SObjectTableRow) are supported

#rnx

diff --git a/Engine/Plugins/Runtime/CommonUI/Source/CommonUI/Private/CommonActivatableWidget.cpp b/Engine/Plugins/Runtime/CommonUI/Source/CommonUI/Private/CommonActivatableWidget.cpp
--- a/Engine/Plugins/Runtime/CommonUI/Source/CommonUI/Private/CommonActivatableWidget.cpp
+++ b/Engine/Plugins/Runtime/CommonUI/Source/CommonUI/Private/CommonActivatableWidget.cpp
@@ -212,12 +212,11 @@ TObjectPtr<UCommonInputActionDomain> UCommonActivatableWidget::GetCalculatedActi
 		return CalculatedActionDomainCache.GetValue().Get();
 	}
 
-	const FName SObjectWidgetName = TEXT("SObjectWidget");
 	const ULocalPlayer* OwningLocalPlayer = GetOwningLocalPlayer();
 	TSharedPtr<SWidget> CurrentWidget = GetCachedWidget();
 	while (CurrentWidget)
 	{
-		if (CurrentWidget->GetType().IsEqual(SObjectWidgetName))
+		if (CurrentWidget->GetMetaData<FCommonActivatableSlateMetaData>().IsValid())
 		{
 			const TSharedPtr<ICommonInputActionDomainMetaData> Metadata = CurrentWidget->GetMetaData<ICommonInputActionDomainMetaData>();
 			if (Metadata.IsValid())
@@ -257,6 +256,17 @@ TSharedRef<SWidget> UCommonActivatableWidget::RebuildWidget()
 	return Super::RebuildWidget();
 }
 
+void UCommonActivatableWidget::OnWidgetRebuilt()
+{
+	Super::OnWidgetRebuilt();
+
+	TSharedPtr<SObjectWidget> GCWidget = MyGCWidget.Pin();
+	if (GCWidget.IsValid() && !GCWidget->GetMetaData<FCommonActivatableSlateMetaData>().IsValid())
+	{
+		GCWidget->AddMetadata(MakeShared<FCommonActivatableSlateMetaData>());
+	}
+}
+
 void UCommonActivatableWidget::ReleaseSlateResources(bool bReleaseChildren)
 {
 	Super::ReleaseSlateResources(bReleaseChildren);
diff --git a/Engine/Plugins/Runtime/CommonUI/Source/CommonUI/Private/Input/CommonUIActionRouterBase.cpp b/Engine/Plugins/Runtime/CommonUI/Source/CommonUI/Private/Input/CommonUIActionRouterBase.cpp
--- a/Engine/Plugins/Runtime/CommonUI/Source/CommonUI/Private/Input/CommonUIActionRouterBase.cpp
+++ b/Engine/Plugins/Runtime/CommonUI/Source/CommonUI/Private/Input/CommonUIActionRouterBase.cpp
@@ -134,32 +134,73 @@ public:
 // UCommonUIActionRouterBase
 //////////////////////////////////////////////////////////////////////////
 
-UCommonUIActionRouterBase* UCommonUIActionRouterBase::Get(const UWidget& ContextWidget)
-{
-	return ULocalPlayer::GetSubsystem<UCommonUIActionRouterBase>(ContextWidget.GetOwningLocalPlayer());
-}
-
-UCommonActivatableWidget* UCommonUIActionRouterBase::FindOwningActivatable(TSharedPtr<SWidget> Widget, ULocalPlayer* OwningLocalPlayer)
+namespace UE::CommonUI::Private
 {
-	UCommonActivatableWidget* OwningActivatable = nullptr;
+	enum class ESearchType
+	{
+		IncludeSelf,
+		ExcludeSelf
+	};
+	template<ESearchType SearchType, typename TPredicate>
+	void ForEachParentWidget(TSharedRef<SWidget> Widget, TPredicate Predicate)
+	{
+		TSharedPtr<SWidget> TestWidget = Widget;
+		if constexpr (SearchType == ESearchType::ExcludeSelf)
+		{
+			TestWidget = Widget->GetParentWidget();
+		}
+		while (TestWidget)
+		{
+			const bool bContinueIterating = Predicate(TestWidget.ToSharedRef());
+			if (!bContinueIterating)
+			{
+				return;
+			}
+			
+			TestWidget = TestWidget->GetParentWidget();
+		}
+	}
 
-	while (Widget && !OwningActivatable)
+	
+	template<ESearchType SearchType>
+	UCommonActivatableWidget* FindActivatableFromSlate(TSharedPtr<SWidget> SlateWidget, ULocalPlayer* OwningLocalPlayer)
 	{
-		Widget = Widget->GetParentWidget();
-		if (Widget && Widget->GetType().IsEqual(TEXT("SObjectWidget")))
+		UCommonActivatableWidget* OwningActivatable = nullptr;
+		if (SlateWidget.IsValid())
 		{
-			if (UCommonActivatableWidget* CandidateActivatable = Cast<UCommonActivatableWidget>(StaticCastSharedPtr<SObjectWidget>(Widget)->GetWidgetObject()))
+			ForEachParentWidget<SearchType>(SlateWidget.ToSharedRef(), [&OwningActivatable, OwningLocalPlayer](const TSharedRef<SWidget>& Widget)
 			{
-				if (CandidateActivatable->GetOwningLocalPlayer() != OwningLocalPlayer)
+				if (Widget->GetMetaData<FCommonActivatableSlateMetaData>().IsValid())
 				{
-					return nullptr;
+					if (UCommonActivatableWidget* CandidateActivatable = Cast<UCommonActivatableWidget>(StaticCastSharedRef<SObjectWidget>(Widget)->GetWidgetObject()))
+					{
+						if (OwningLocalPlayer == nullptr || CandidateActivatable->GetOwningLocalPlayer() == OwningLocalPlayer)
+						{
+							OwningActivatable = CandidateActivatable;
+						}
+						return false;
+					}
 				}
-				OwningActivatable = CandidateActivatable;
-			}
+				return true;
+			});
 		}
+		return OwningActivatable;
 	}
+}
+
+UCommonUIActionRouterBase* UCommonUIActionRouterBase::Get(const UWidget& ContextWidget)
+{
+	return ULocalPlayer::GetSubsystem<UCommonUIActionRouterBase>(ContextWidget.GetOwningLocalPlayer());
+}
 
-	return OwningActivatable;
+UCommonActivatableWidget* UCommonUIActionRouterBase::FindOwningActivatable(TSharedPtr<SWidget> Widget, ULocalPlayer* OwningLocalPlayer)
+{
+	return UE::CommonUI::Private::FindActivatableFromSlate<UE::CommonUI::Private::ESearchType::ExcludeSelf>(Widget, OwningLocalPlayer);
+}
+
+UCommonActivatableWidget* UCommonUIActionRouterBase::FindActivatable(TSharedPtr<SWidget> Widget, ULocalPlayer* OwningLocalPlayer)
+{
+	return UE::CommonUI::Private::FindActivatableFromSlate<UE::CommonUI::Private::ESearchType::IncludeSelf>(Widget, OwningLocalPlayer);
 }
 
 UCommonUIActionRouterBase::UCommonUIActionRouterBase()
@@ -649,7 +690,7 @@ bool UCommonUIActionRouterBase::IsWidgetInActiveRoot(const UCommonActivatableWid
 		TSharedPtr<SWidget> WidgetWalker = Widget->GetCachedWidget();
 		while (WidgetWalker)
 		{
-			if (WidgetWalker->GetType().IsEqual(TEXT("SObjectWidget")))
+			if (WidgetWalker->GetMetaData<FCommonActivatableSlateMetaData>().IsValid())
 			{
 				if (UCommonActivatableWidget* CandidateActivatable = Cast<UCommonActivatableWidget>(StaticCastSharedPtr<SObjectWidget>(WidgetWalker)->GetWidgetObject()))
 				{
@@ -903,9 +944,13 @@ void UCommonUIActionRouterBase::HandleLeafmostActiveNodeChanged()
 
 void UCommonUIActionRouterBase::HandleSlateFocusChanging(const FFocusEvent& FocusEvent, const FWeakWidgetPath& OldFocusedWidgetPath, const TSharedPtr<SWidget>& OldFocusedWidget, const FWidgetPath& NewFocusedWidgetPath, const TSharedPtr<SWidget>& NewFocusedWidget)
 {
-	if (FocusEvent.GetCause() == EFocusCause::SetDirectly && FocusEvent.GetUser() == GetLocalPlayerIndex() && ActiveRootNode && ActiveRootNode->IsExclusiveParentOfWidget(OldFocusedWidget))
+	if (FocusEvent.GetUser() == GetLocalPlayerIndex())
 	{
-		ActiveRootNode->RefreshCachedRestorationTarget();
+		FActivatableTreeRootPtr RootNode = ActiveRootNode ? ActiveRootNode : FindActiveActionDomainRootNode();
+		if (RootNode && RootNode->IsParentOfWidget(OldFocusedWidget, FActivatableTreeNode::IncludeSelf))
+		{
+			RootNode->RefreshCachedRestorationTarget();
+		}
 	}
 }
 
diff --git a/Engine/Plugins/Runtime/CommonUI/Source/CommonUI/Private/Input/UIActionRouterTypes.cpp b/Engine/Plugins/Runtime/CommonUI/Source/CommonUI/Private/Input/UIActionRouterTypes.cpp
--- a/Engine/Plugins/Runtime/CommonUI/Source/CommonUI/Private/Input/UIActionRouterTypes.cpp
+++ b/Engine/Plugins/Runtime/CommonUI/Source/CommonUI/Private/Input/UIActionRouterTypes.cpp
@@ -1044,10 +1044,11 @@ FActivatableTreeNodeRef FActivatableTreeNode::AddChildNode(UCommonActivatableWid
 void FActivatableTreeNode::CacheFocusRestorationTarget()
 {
 	TSharedPtr<SWidget> FocusedWidget = FSlateApplication::Get().GetUserFocusedWidget(GetOwnerUserIndex());
-	UCommonActivatableWidget* FocusedActivatableWidget = FocusedWidget ? UCommonUIActionRouterBase::FindOwningActivatable(FocusedWidget, GetActionRouter().GetLocalPlayerChecked()) : nullptr;
+	UCommonActivatableWidget* FocusedActivatableWidget = FocusedWidget ? UCommonUIActionRouterBase::FindActivatable(FocusedWidget, GetActionRouter().GetLocalPlayerChecked()) : nullptr;
 
-	if (FocusedWidget != FocusRestorationTarget.Pin() &&  (!FocusedActivatableWidget || IsWidgetInNodeHierarchy(FocusedActivatableWidget, *this)))
+	if (FocusedWidget != FocusRestorationTarget.Pin() && (!FocusedWidget || (FocusedActivatableWidget && IsWidgetInNodeHierarchy(FocusedActivatableWidget, *this))))
 	{
+		UE_LOG(LogUIActionRouter, VeryVerbose, TEXT("FocusRestorationTarget for [%s] changed from [%s] to [%s]"), *GetNameSafe(GetWidget()), *FReflectionMetaData::GetWidgetDebugInfo(FocusRestorationTarget.Pin().Get()), *FReflectionMetaData::GetWidgetDebugInfo(FocusedWidget.Get()));
 		FocusRestorationTarget = FocusedWidget;
 	}
 }
@@ -1211,12 +1212,12 @@ void FActivatableTreeNode::DebugDumpRecursive(FString& OutputStr, int32 Depth, b
 	}
 }
 
-bool FActivatableTreeNode::IsParentOfWidget(const TSharedPtr<SWidget>& SlateWidget) const
+bool FActivatableTreeNode::IsParentOfWidget(const TSharedPtr<SWidget>& SlateWidget, EIsParentSearchType ParentSearchType) const
 {
 	if (SlateWidget && ensure(RepresentedWidget.IsValid()))
 	{
 		TSharedPtr<SWidget> CachedWidget = RepresentedWidget->GetCachedWidget();
-		TSharedPtr<SWidget> ParentWidget = SlateWidget->GetParentWidget();
+		TSharedPtr<SWidget> ParentWidget = (ParentSearchType == ExcludeSelf) ? SlateWidget->GetParentWidget() : SlateWidget;
 		while (ParentWidget && ParentWidget != CachedWidget)
 		{
 			ParentWidget = ParentWidget->GetParentWidget();
@@ -1228,11 +1229,11 @@ bool FActivatableTreeNode::IsParentOfWidget(const TSharedPtr<SWidget>& SlateWidg
 
 bool FActivatableTreeNode::IsExclusiveParentOfWidget(const TSharedPtr<SWidget>& SlateWidget) const
 {
-	if (IsParentOfWidget(SlateWidget))
+	if (IsParentOfWidget(SlateWidget, ExcludeSelf))
 	{
 		for (const FActivatableTreeNodeRef& ChildNode : GetChildren())
 		{
-			if (ChildNode->DoesWidgetSupportActivationFocus() && ChildNode->IsParentOfWidget(SlateWidget))
+			if (ChildNode->DoesWidgetSupportActivationFocus() && ChildNode->IsParentOfWidget(SlateWidget, ExcludeSelf))
 			{
 				return false;
 			}
@@ -1625,7 +1626,7 @@ void FActivatableTreeRoot::FocusLeafmostNode()
 
 	if (TSharedPtr<SWidget> AutoRestoreTarget = LeafWidget->AutoRestoresFocus() ? PinnedLeafmostNode->GetFocusFallbackTarget() : nullptr)
 	{
-		UE_LOG(LogUIActionRouter, Display, TEXT("[User %d] Set AutoRestoreTarget"), OwnerSlateId);
+		UE_LOG(LogUIActionRouter, Display, TEXT("[User %d] Focused AutoRestoreTarget %s"), OwnerSlateId, *FReflectionMetaData::GetWidgetDebugInfo(AutoRestoreTarget.Get()));
 		
 		bShouldCancelDelayedFocusOperation = FSlateApplication::Get().SetUserFocus(OwnerSlateId, AutoRestoreTarget);
 	}
diff --git a/Engine/Plugins/Runtime/CommonUI/Source/CommonUI/Private/Input/UIActionRouterTypes.h b/Engine/Plugins/Runtime/CommonUI/Source/CommonUI/Private/Input/UIActionRouterTypes.h
--- a/Engine/Plugins/Runtime/CommonUI/Source/CommonUI/Private/Input/UIActionRouterTypes.h
+++ b/Engine/Plugins/Runtime/CommonUI/Source/CommonUI/Private/Input/UIActionRouterTypes.h
@@ -119,6 +119,12 @@ public:
 	TSharedPtr<SWidget> GetFocusFallbackTarget() const;
 
 	bool IsExclusiveParentOfWidget(const TSharedPtr<SWidget>& SlateWidget) const;
+	enum EIsParentSearchType
+	{
+		ExcludeSelf,
+		IncludeSelf
+	};
+	bool IsParentOfWidget(const TSharedPtr<SWidget>& SlateWidget, EIsParentSearchType ParentSearchType) const;
 
 	int32 GetLastPaintLayer() const;
 	TOptional<FUIInputConfig> FindDesiredInputConfig() const;
@@ -147,8 +153,6 @@ protected:
 	void AppendValidScrollRecipients(TArray<const UWidget*>& AllScrollRecipients) const;
 	void DebugDumpRecursive(FString& OutputStr, int32 Depth, bool bIncludeActions, bool bIncludeChildren, bool bIncludeInactive) const;
 	
-	bool IsParentOfWidget(const TSharedPtr<SWidget>& SlateWidget) const;
-	
 private:
 	void HandleWidgetActivated();
 	void HandleWidgetDeactivated();
diff --git a/Engine/Plugins/Runtime/CommonUI/Source/CommonUI/Private/Widgets/CommonActivatableWidgetContainer.cpp b/Engine/Plugins/Runtime/CommonUI/Source/CommonUI/Private/Widgets/CommonActivatableWidgetContainer.cpp
--- a/Engine/Plugins/Runtime/CommonUI/Source/CommonUI/Private/Widgets/CommonActivatableWidgetContainer.cpp
+++ b/Engine/Plugins/Runtime/CommonUI/Source/CommonUI/Private/Widgets/CommonActivatableWidgetContainer.cpp
@@ -13,7 +13,7 @@
 
 UCommonActivatableWidget* ActivatableWidgetFromSlate(const TSharedPtr<SWidget>& SlateWidget)
 {
-	if (SlateWidget && SlateWidget != SNullWidget::NullWidget && ensure(SlateWidget->GetType().IsEqual(TEXT("SObjectWidget"))))
+	if (SlateWidget && SlateWidget != SNullWidget::NullWidget && ensure(SlateWidget->GetMetaData<FCommonActivatableSlateMetaData>().IsValid()))
 	{
 		UCommonActivatableWidget* ActivatableWidget = Cast<UCommonActivatableWidget>(StaticCastSharedPtr<SObjectWidget>(SlateWidget)->GetWidgetObject());
 		if (ensure(ActivatableWidget))
diff --git a/Engine/Plugins/Runtime/CommonUI/Source/CommonUI/Public/CommonActivatableWidget.h b/Engine/Plugins/Runtime/CommonUI/Source/CommonUI/Public/CommonActivatableWidget.h
--- a/Engine/Plugins/Runtime/CommonUI/Source/CommonUI/Public/CommonActivatableWidget.h
+++ b/Engine/Plugins/Runtime/CommonUI/Source/CommonUI/Public/CommonActivatableWidget.h
@@ -12,6 +12,15 @@ class UInputMappingContext;
 
 DECLARE_DYNAMIC_MULTICAST_DELEGATE(FOnWidgetActivationChanged);
 
+/**
+ * Metadata set on UCommonActivatableWidget's underlying slate widget to allow Slate -> UMG lookup
+ */
+class FCommonActivatableSlateMetaData : public ISlateMetaData
+{
+public:
+	SLATE_METADATA_TYPE(FCommonActivatableSlateMetaData, ISlateMetaData);
+};
+
 /** 
  * The base for widgets that are capable of being "activated" and "deactivated" during their lifetime without being otherwise modified or destroyed. 
  *
@@ -111,6 +120,7 @@ public:
 
 protected:
 	virtual TSharedRef<SWidget> RebuildWidget() override;
+	virtual void OnWidgetRebuilt() override;
 	virtual void ReleaseSlateResources(bool bReleaseChildren) override;
 	virtual void NativeConstruct() override;
 	virtual void NativeDestruct() override;
diff --git a/Engine/Plugins/Runtime/CommonUI/Source/CommonUI/Public/Input/CommonUIActionRouterBase.h b/Engine/Plugins/Runtime/CommonUI/Source/CommonUI/Public/Input/CommonUIActionRouterBase.h
--- a/Engine/Plugins/Runtime/CommonUI/Source/CommonUI/Public/Input/CommonUIActionRouterBase.h
+++ b/Engine/Plugins/Runtime/CommonUI/Source/CommonUI/Public/Input/CommonUIActionRouterBase.h
@@ -68,8 +68,10 @@ class COMMONUI_API UCommonUIActionRouterBase : public ULocalPlayerSubsystem
 public:
 	static UCommonUIActionRouterBase* Get(const UWidget& ContextWidget);
 
-	/** searches up the SWidget tree until it finds the nearest UCommonActivatableWidget */
+	/** searches up the SWidget tree until it finds the nearest UCommonActivatableWidget (excluding checking Widget itself) */
 	static UCommonActivatableWidget* FindOwningActivatable(TSharedPtr<SWidget> Widget, ULocalPlayer* OwningLocalPlayer);
+	/** searches up the SWidget tree until it finds the nearest UCommonActivatableWidget (including checking Widget itself) */
+	static UCommonActivatableWidget* FindActivatable(TSharedPtr<SWidget> Widget, ULocalPlayer* OwningLocalPlayer);
 
 	UCommonUIActionRouterBase();
 	virtual void Initialize(FSubsystemCollectionBase& Collection) override;
