Update CommonUI preprocessor to switch input methods when the mouse wheel is scrolled

diff --git a/Engine/Plugins/Runtime/CommonUI/Source/CommonInput/Private/CommonInputPreprocessor.cpp b/Engine/Plugins/Runtime/CommonUI/Source/CommonInput/Private/CommonInputPreprocessor.cpp
--- a/Engine/Plugins/Runtime/CommonUI/Source/CommonInput/Private/CommonInputPreprocessor.cpp
+++ b/Engine/Plugins/Runtime/CommonUI/Source/CommonInput/Private/CommonInputPreprocessor.cpp
@@ -117,6 +117,20 @@ bool FCommonInputPreprocessor::HandleMouseButtonDoubleClickEvent(FSlateApplicati
 	return false;
 }
 
+bool FCommonInputPreprocessor::HandleMouseWheelOrGestureEvent(FSlateApplication& SlateApp, const FPointerEvent& InWheelEvent, const FPointerEvent* InGestureEvent)
+{
+	const ECommonInputType InputType = GetInputType(InWheelEvent);
+	if (IsRelevantInput(SlateApp, InWheelEvent, InputType))
+	{
+		if (IsInputMethodBlocked(InputType))
+		{
+			return true;
+		}
+		RefreshCurrentInputMethod(InputType);
+	}
+	return false;
+}
+
 void FCommonInputPreprocessor::SetInputTypeFilter(ECommonInputType InputType, FName InReason, bool InFilter)
 {
 	TMap<FName, bool>& Reasons = FilterInputTypeWithReasons[(uint8)InputType];
diff --git a/Engine/Plugins/Runtime/CommonUI/Source/CommonInput/Public/CommonInputPreprocessor.h b/Engine/Plugins/Runtime/CommonUI/Source/CommonInput/Public/CommonInputPreprocessor.h
--- a/Engine/Plugins/Runtime/CommonUI/Source/CommonInput/Public/CommonInputPreprocessor.h
+++ b/Engine/Plugins/Runtime/CommonUI/Source/CommonInput/Public/CommonInputPreprocessor.h
@@ -25,6 +25,7 @@ public:
 	virtual bool HandleMouseMoveEvent(FSlateApplication& SlateApp, const FPointerEvent& InPointerEvent) override;
 	virtual bool HandleMouseButtonDownEvent(FSlateApplication& SlateApp, const FPointerEvent& InPointerEvent) override;
 	virtual bool HandleMouseButtonDoubleClickEvent(FSlateApplication& SlateApp, const FPointerEvent& InPointerEvent) override;
+	virtual bool HandleMouseWheelOrGestureEvent(FSlateApplication& SlateApp, const FPointerEvent& InWheelEvent, const FPointerEvent* InGestureEvent) override;
 	virtual const TCHAR* GetDebugName() const override { return TEXT("CommonInput"); }
 	//~ End IInputProcessor Interface
 
