Adding support for multiple ControllerData assets per controller in UCommonInputPlatformSettings

#rnx

diff --git a/Engine/Plugins/Runtime/CommonUI/Source/CommonInput/Private/CommonInputBaseTypes.cpp b/Engine/Plugins/Runtime/CommonUI/Source/CommonInput/Private/CommonInputBaseTypes.cpp
--- a/Engine/Plugins/Runtime/CommonUI/Source/CommonInput/Private/CommonInputBaseTypes.cpp
+++ b/Engine/Plugins/Runtime/CommonUI/Source/CommonInput/Private/CommonInputBaseTypes.cpp
@@ -287,9 +287,12 @@ bool UCommonInputPlatformSettings::TryGetInputBrush(FSlateBrush& OutBrush, FKey
 {
 	InitializeControllerData();
 
-	if (const UCommonInputBaseControllerData* DefaultControllerData = GetControllerDataForInputType(InputType, GamepadName))
+	for (const UCommonInputBaseControllerData* DefaultControllerData: GetControllerDataForInputType(InputType, GamepadName))
 	{
-		return DefaultControllerData->TryGetInputBrush(OutBrush, Key);
+		if (DefaultControllerData->TryGetInputBrush(OutBrush, Key))
+		{
+			return true;
+		}
 	}
 
 	return false;
@@ -299,18 +302,22 @@ bool UCommonInputPlatformSettings::TryGetInputBrush(FSlateBrush& OutBrush, const
 {
 	InitializeControllerData();
 
-	if (const UCommonInputBaseControllerData* DefaultControllerData = GetControllerDataForInputType(InputType, GamepadName))
+	for (const UCommonInputBaseControllerData* DefaultControllerData: GetControllerDataForInputType(InputType, GamepadName))
 	{
-		return DefaultControllerData->TryGetInputBrush(OutBrush, Keys);
+		if (DefaultControllerData->TryGetInputBrush(OutBrush, Keys))
+		{
+			return true;
+		}
 	}
 
 	return false;
 }
 
-const UCommonInputBaseControllerData* UCommonInputPlatformSettings::GetControllerDataForInputType(ECommonInputType InputType, const FName GamepadName) const
+TArray<const UCommonInputBaseControllerData*> UCommonInputPlatformSettings::GetControllerDataForInputType(ECommonInputType InputType, const FName GamepadName) const
 {
 	InitializeControllerData();
 
+	TArray<const UCommonInputBaseControllerData*> ControllerDataForInputType;
 	for (const TSubclassOf<UCommonInputBaseControllerData>& ControllerDataPtr : ControllerDataClasses)
 	{
 		const UCommonInputBaseControllerData* DefaultControllerData = ControllerDataPtr.GetDefaultObject();
@@ -318,12 +325,12 @@ const UCommonInputBaseControllerData* UCommonInputPlatformSettings::GetControlle
 		{
 			if (DefaultControllerData->InputType != ECommonInputType::Gamepad || DefaultControllerData->GamepadName == GamepadName)
 			{
-				return DefaultControllerData;
+				ControllerDataForInputType.Add(DefaultControllerData);
 			}
 		}
 	}
 
-	return nullptr;
+	return ControllerDataForInputType;
 }
 
 FName UCommonInputPlatformSettings::GetBestGamepadNameForHardware(FName CurrentGamepadName, FName InputDeviceName, const FString& HardwareDeviceIdentifier)
diff --git a/Engine/Plugins/Runtime/CommonUI/Source/CommonInput/Public/CommonInputBaseTypes.h b/Engine/Plugins/Runtime/CommonUI/Source/CommonInput/Public/CommonInputBaseTypes.h
--- a/Engine/Plugins/Runtime/CommonUI/Source/CommonInput/Public/CommonInputBaseTypes.h
+++ b/Engine/Plugins/Runtime/CommonUI/Source/CommonInput/Public/CommonInputBaseTypes.h
@@ -230,7 +230,7 @@ public:
 	bool TryGetInputBrush(FSlateBrush& OutBrush, FKey Key, ECommonInputType InputType, const FName GamepadName) const;
 	bool TryGetInputBrush(FSlateBrush& OutBrush, const TArray<FKey>& Keys, ECommonInputType InputType, const FName GamepadName) const;
 
-	const UCommonInputBaseControllerData* GetControllerDataForInputType(ECommonInputType InputType, const FName GamepadName) const;
+	TArray<const UCommonInputBaseControllerData*> GetControllerDataForInputType(ECommonInputType InputType, const FName GamepadName) const;
 
 	FName GetBestGamepadNameForHardware(FName CurrentGamepadName, FName InputDeviceName, const FString& HardwareDeviceIdentifier); 
 
