CommonUI: Fix scrolling with right stick causing input type to change to MouseAndKeyboard

#rnx

diff --git a/Engine/Plugins/Runtime/CommonUI/Source/CommonUI/Private/Input/CommonAnalogCursor.cpp b/Engine/Plugins/Runtime/CommonUI/Source/CommonUI/Private/Input/CommonAnalogCursor.cpp
--- a/Engine/Plugins/Runtime/CommonUI/Source/CommonUI/Private/Input/CommonAnalogCursor.cpp
+++ b/Engine/Plugins/Runtime/CommonUI/Source/CommonUI/Private/Input/CommonAnalogCursor.cpp
@@ -201,7 +201,10 @@ void FCommonAnalogCursor::Tick(const float DeltaTime, FSlateApplication& SlateAp
 										ScrollAmount,
 										FModifierKeysState());
 
+									UCommonInputSubsystem& InputSubsytem = ActionRouter.GetInputSubsystem();
+									InputSubsytem.SetIsGamepadSimulatedClick(true);
 									SlateApp.ProcessMouseWheelOrGestureEvent(MouseEvent, nullptr);
+									InputSubsytem.SetIsGamepadSimulatedClick(false);
 								}
 							}
 						}
