CommonUI: Unify the path for validating which input actions are displayed in CommonBoundActionBar for Enhanced Input actions and Data Table actions

- This means Data Table Touch and KeyboardAndMouse bindings are now validating their keys before being displayed when previously, only Gamepad bindings were validated

#rnx

diff --git a/Engine/Plugins/Runtime/CommonUI/Source/CommonUI/Private/CommonUITypes.cpp b/Engine/Plugins/Runtime/CommonUI/Source/CommonUI/Private/CommonUITypes.cpp
--- a/Engine/Plugins/Runtime/CommonUI/Source/CommonUI/Private/CommonUITypes.cpp
+++ b/Engine/Plugins/Runtime/CommonUI/Source/CommonUI/Private/CommonUITypes.cpp
@@ -302,7 +302,7 @@ FSlateBrush CommonUI::GetIconForEnhancedInputAction(const UCommonInputSubsystem*
 
 bool CommonUI::ActionValidForInputType(const ULocalPlayer* LocalPlayer, ECommonInputType InputType, const UInputAction* InputAction)
 {
-	if (!LocalPlayer)
+	if (!LocalPlayer || !InputAction)
 	{
 		return false;
 	}
@@ -315,16 +315,7 @@ bool CommonUI::ActionValidForInputType(const ULocalPlayer* LocalPlayer, ECommonI
 
 	for (const FKey& Key : Keys)
 	{
-		if (!Key.IsValid())
-		{
-			continue;
-		}
-
-		bool bIsValidTouch = Key.IsTouch() && InputType == ECommonInputType::Touch;
-		bool bIsValidGamepad = Key.IsGamepadKey() && InputType == ECommonInputType::Gamepad;
-		bool bIsValidMouseAndKeyboard = !Key.IsTouch() && !Key.IsGamepadKey() && InputType == ECommonInputType::MouseAndKeyboard;
-
-		if (bIsValidTouch || bIsValidGamepad || bIsValidMouseAndKeyboard)
+		if (IsKeyValidForInputType(Key, InputType))
 		{
 			return true;
 		}
@@ -335,7 +326,7 @@ bool CommonUI::ActionValidForInputType(const ULocalPlayer* LocalPlayer, ECommonI
 
 FKey CommonUI::GetFirstKeyForInputType(const ULocalPlayer* LocalPlayer, ECommonInputType InputType, const UInputAction* InputAction)
 {
-	if (!LocalPlayer)
+	if (!LocalPlayer || !InputAction)
 	{
 		return FKey();
 	}
@@ -348,20 +339,7 @@ FKey CommonUI::GetFirstKeyForInputType(const ULocalPlayer* LocalPlayer, ECommonI
 
 	for (const FKey& Key : Keys)
 	{
-		if (!Key.IsValid())
-		{
-			continue;
-		}
-
-		if (Key.IsTouch() && InputType == ECommonInputType::Touch)
-		{
-			return Key;
-		}
-		else if (Key.IsGamepadKey() && InputType == ECommonInputType::Gamepad)
-		{
-			return Key;
-		}
-		else if (!Key.IsTouch() && !Key.IsGamepadKey() && InputType == ECommonInputType::MouseAndKeyboard)
+		if (IsKeyValidForInputType(Key, InputType))
 		{
 			return Key;
 		}
@@ -369,3 +347,32 @@ FKey CommonUI::GetFirstKeyForInputType(const ULocalPlayer* LocalPlayer, ECommonI
 
 	return FKey();
 }
+
+bool CommonUI::ActionValidForInputType(const ULocalPlayer* LocalPlayer, ECommonInputType InputType, const FCommonInputActionDataBase* InputAction)
+{
+	if (!LocalPlayer || !InputAction)
+	{
+		return false;
+	}
+	
+	UCommonInputSubsystem* InputSubsystem = LocalPlayer->GetSubsystem<UCommonInputSubsystem>();
+	const FCommonInputTypeInfo& TypeInfo = InputAction->GetInputTypeInfo(InputType, InputSubsystem->GetCurrentGamepadName());
+	return IsKeyValidForInputType(TypeInfo.GetKey(), InputType);
+}
+
+bool CommonUI::IsKeyValidForInputType(const FKey& Key, ECommonInputType InputType)
+{
+	switch (InputType)
+	{
+	case ECommonInputType::MouseAndKeyboard:
+		return !Key.IsTouch() && !Key.IsGamepadKey();
+	case ECommonInputType::Gamepad:
+		return Key.IsGamepadKey();
+	case ECommonInputType::Touch:
+		return Key.IsTouch();
+	case ECommonInputType::Count:
+	default:
+		checkNoEntry(); // Unhandled type
+		return false;
+	}
+}
diff --git a/Engine/Plugins/Runtime/CommonUI/Source/CommonUI/Private/Input/CommonBoundActionBar.cpp b/Engine/Plugins/Runtime/CommonUI/Source/CommonUI/Private/Input/CommonBoundActionBar.cpp
--- a/Engine/Plugins/Runtime/CommonUI/Source/CommonUI/Private/Input/CommonBoundActionBar.cpp
+++ b/Engine/Plugins/Runtime/CommonUI/Source/CommonUI/Private/Input/CommonBoundActionBar.cpp
@@ -175,35 +175,15 @@ void UCommonBoundActionBar::HandleDeferredDisplayUpdate()
 									return false;
 								}
 
-								if (CommonUI::IsEnhancedInputSupportEnabled())
+								const bool bIsValidEnhancedInputAction = CommonUI::IsEnhancedInputSupportEnabled() && CommonUI::ActionValidForInputType(ActionRouter->GetLocalPlayer(), PlayerInputType, Binding->InputAction.Get());
+								if (!bIsValidEnhancedInputAction)
 								{
-									if (TObjectPtr<const UInputAction> InputAction = Binding->InputAction.Get())
-									{
-										if (CommonUI::ActionValidForInputType(ActionRouter->GetLocalPlayer(), PlayerInputType, InputAction))
-										{
-											if (!bIgnoreDuplicateActions)
-											{
-												return true;
-											}
-											bool bAlreadyAccepted = false;
-											AcceptedBindings.Add(Binding->ActionName, &bAlreadyAccepted);
-											return !bAlreadyAccepted;
-										}
-										return false;
-									}
-								}
-
-								if (FCommonInputActionDataBase* LegacyData = Binding->GetLegacyInputActionData())
-								{
-									if (!LegacyData->CanDisplayInReflector(PlayerInputType, PlayerGamepadName))
+									const bool bIsValidDataTableInputAction = CommonUI::ActionValidForInputType(ActionRouter->GetLocalPlayer(), PlayerInputType, Binding->GetLegacyInputActionData());
+									if (!bIsValidDataTableInputAction)
 									{
 										return false;
 									}
 								}
-								else
-								{
-									return false; 
-								}
 
 								if (!bIgnoreDuplicateActions)
 								{
diff --git a/Engine/Plugins/Runtime/CommonUI/Source/CommonUI/Public/CommonUITypes.h b/Engine/Plugins/Runtime/CommonUI/Source/CommonUI/Public/CommonUITypes.h
--- a/Engine/Plugins/Runtime/CommonUI/Source/CommonUI/Public/CommonUITypes.h
+++ b/Engine/Plugins/Runtime/CommonUI/Source/CommonUI/Public/CommonUITypes.h
@@ -288,6 +288,8 @@ public:
 	static void InjectEnhancedInputForAction(const ULocalPlayer* LocalPlayer, const UInputAction* InputAction, FInputActionValue RawValue);
 	static FSlateBrush GetIconForEnhancedInputAction(const UCommonInputSubsystem* CommonInputSubsystem, const UInputAction* InputAction);
 	static bool ActionValidForInputType(const ULocalPlayer* LocalPlayer, ECommonInputType InputType, const UInputAction* InputAction);
+	static bool ActionValidForInputType(const ULocalPlayer* LocalPlayer, ECommonInputType InputType, const FCommonInputActionDataBase* InputAction);
+	static bool IsKeyValidForInputType(const FKey& Key, ECommonInputType InputType);
 	static FKey GetFirstKeyForInputType(const ULocalPlayer* LocalPlayer, ECommonInputType InputType, const UInputAction* InputAction);
 };
 
