CommonUI: Prevent needlessly setting cursor position every frame when the viewport is the cursor target due to incorrect change target geometry logic

#rnx

diff --git a/Engine/Plugins/Runtime/CommonUI/Source/CommonUI/Private/Input/CommonAnalogCursor.cpp b/Engine/Plugins/Runtime/CommonUI/Source/CommonUI/Private/Input/CommonAnalogCursor.cpp
--- a/Engine/Plugins/Runtime/CommonUI/Source/CommonUI/Private/Input/CommonAnalogCursor.cpp
+++ b/Engine/Plugins/Runtime/CommonUI/Source/CommonUI/Private/Input/CommonAnalogCursor.cpp
@@ -98,9 +98,27 @@ void FCommonAnalogCursor::Tick(const float DeltaTime, FSlateApplication& SlateAp
 					CursorTarget = SelectedRows[0]->AsWidget();
 				}
 			}
+			
+			FGeometry TargetGeometry;
+			if (CursorTarget)
+			{
+				if (CursorTarget == GetViewportClient()->GetGameViewportWidget())
+				{
+					// When the target is the game viewport as a whole, we don't want to center blindly - we want to center in the geometry of our owner's widget host layer
+					TSharedPtr<IGameLayerManager> GameLayerManager = GetViewportClient()->GetGameLayerManager();
+					if (ensure(GameLayerManager))
+					{
+						TargetGeometry = GameLayerManager->GetPlayerWidgetHostGeometry(ActionRouter.GetLocalPlayerChecked());
+					}
+				}
+				else
+				{
+					TargetGeometry = CursorTarget->GetTickSpaceGeometry();
+				}
+			}
 
 			// We want to try to update the cursor position when focus changes or the focused widget moves at all
-			if (CursorTarget != PinnedLastCursorTarget || (CursorTarget && CursorTarget->GetCachedGeometry().GetAccumulatedRenderTransform() != LastCursorTargetTransform))
+			if (CursorTarget != PinnedLastCursorTarget || (CursorTarget && TargetGeometry.GetAccumulatedRenderTransform() != LastCursorTargetTransform))
 			{
 #if !UE_BUILD_SHIPPING
 				if (CursorTarget != PinnedLastCursorTarget)
@@ -121,21 +139,6 @@ void FCommonAnalogCursor::Tick(const float DeltaTime, FSlateApplication& SlateAp
 				bool bHasValidCursorTarget = false;
 				if (CursorTarget)
 				{
-					FGeometry TargetGeometry; 
-					if (CursorTarget == GetViewportClient()->GetGameViewportWidget())
-					{
-						// When the target is the game viewport as a whole, we don't want to center blindly - we want to center in the geometry of our owner's widget host layer
-						TSharedPtr<IGameLayerManager> GameLayerManager = GetViewportClient()->GetGameLayerManager();
-						if (ensure(GameLayerManager))
-						{
-							TargetGeometry = GameLayerManager->GetPlayerWidgetHostGeometry(ActionRouter.GetLocalPlayerChecked());
-						}
-					}
-					else
-					{
-						TargetGeometry = CursorTarget->GetTickSpaceGeometry();
-					}
-
 					if (TargetGeometry.GetLocalSize().X > UE_SMALL_NUMBER && TargetGeometry.GetLocalSize().Y > UE_SMALL_NUMBER)
 					{
 						LastCursorTargetTransform = TargetGeometry.GetAccumulatedRenderTransform();
