Improved how the property bSimulateHoverOnTouchInput works in CommonButtonBase.
The CVar CommonButton.EnableSimulateHoverOnTouch will accept a new value to support the previous implementation for the moment. The legacy mode will be removed in a future update.
0: Disable
1: Enable (default)
2: Legacy Mode (Deprecated)
[RN] Improved how the property bSimulateHoverOnTouchInput works in CommonButtonBase.

diff --git a/Engine/Plugins/Runtime/CommonUI/Source/CommonUI/Private/CommonButtonBase.cpp b/Engine/Plugins/Runtime/CommonUI/Source/CommonUI/Private/CommonButtonBase.cpp
--- a/Engine/Plugins/Runtime/CommonUI/Source/CommonUI/Private/CommonButtonBase.cpp
+++ b/Engine/Plugins/Runtime/CommonUI/Source/CommonUI/Private/CommonButtonBase.cpp
@@ -32,11 +32,11 @@
 
 namespace UE::CommonUI::Private
 {
-	int32 bEnableSimulateHoverOnTouchInput = 1;
+	int32 EnableSimulateHoverOnTouchInput = 1;
 	FAutoConsoleVariableRef CVarEnableSimulateHoverOnTouchInput(
 		TEXT("CommonButton.EnableSimulateHoverOnTouch"),
-		bEnableSimulateHoverOnTouchInput,
-		TEXT("Allows buttons to simulate hovering on touch in accordance with the property SimulateHoverOnTouchInput.\n0: Disable, 1: Enable (default)"),
+		EnableSimulateHoverOnTouchInput,
+		TEXT("Allows buttons to simulate hovering on touch in accordance with the property SimulateHoverOnTouchInput.\n0: Disable, 1: Enable (default), 2: Legacy Mode (Deprecated)\n\nNote on Legacy Mode: This mode was previously the default (CommonButton.EnableSimulateHoverOnTouch=1) and is preserved for retro-compatibility with CommonButton.EnableSimulateHoverOnTouch=2 but will be removed in a future update. In this legacy implementation, the property SimulateHoverOnTouch=true simulates Hover events as expected. However, SimulateHoverOnTouch=false was not blocking all hover events on touch. This implementation was simply preventing the Press & Release functions from simulating more hover events. This implementation was causing inconsistent behaviors for widgets with SimulateHoverOnTouch=false. The new implementation will effectively block all Hover events for UCommonButtonBase if the property SimulateHoverOnTouch is set to false in the editor. This ensures a consistent behavior for touch input so a widget can be built without any Hover events being simulated on a touch screen."),
 		ECVF_Default);
 }
 
@@ -460,7 +460,7 @@ void UCommonButtonBase::NativeConstruct()
 	{
 		HoldData = ICommonInputModule::GetSettings().GetDefaultHoldData();
 	}
-	
+
 	BindTriggeringInputActionToClick();
 	BindInputMethodChangedDelegate();
 	UpdateInputActionWidget();
@@ -553,6 +553,47 @@ void UCommonButtonBase::OnInputMethodChanged(ECommonInputType CurrentInputType)
 	}
 }
 
+bool UCommonButtonBase::IsHoverSimulationOnTouchAvailable() const
+{
+	return UE::CommonUI::Private::EnableSimulateHoverOnTouchInput != 0;
+}
+
+bool UCommonButtonBase::ShouldProcessHoverEvent(EHoverEventSource HoverReason)
+{
+	if (const UCommonInputSubsystem* CommonInputSubsystem = UCommonInputSubsystem::Get(GetOwningLocalPlayer()))
+	{
+		if (CommonInputSubsystem->GetCurrentInputType() == ECommonInputType::Touch)
+		{
+			if (HoverReason == EHoverEventSource::SimulationForTouch)
+			{
+				// disabled mode: blocks internally simulated Hover events
+				if (UE::CommonUI::Private::EnableSimulateHoverOnTouchInput == 0)
+				{
+					return false;
+				}
+
+				// legacy mode: blocks internally simulated Hover events when bSimulateHoverOnTouchInput is false
+				if (UE::CommonUI::Private::EnableSimulateHoverOnTouchInput == 2 && !bSimulateHoverOnTouchInput)
+				{
+					return false;
+				}
+			}
+
+			// strict mode: blocks ALL Hover events when bSimulateHoverOnTouchInput is false
+			if (UE::CommonUI::Private::EnableSimulateHoverOnTouchInput == 1 && !bSimulateHoverOnTouchInput)
+			{
+				return false;
+			}
+
+			// If we do not explicitly want to block hover events on touch for that button, it should get processed.
+			return true;
+		}
+	}
+
+	// On all input methods except touch, process all events except the ones simulated for touch.
+	return HoverReason != EHoverEventSource::SimulationForTouch;
+}
+
 void UCommonButtonBase::UpdateHoldData(ECommonInputType CurrentInputType)
 {
 	if (HoldData && bRequiresHold)
@@ -610,7 +651,7 @@ void UCommonButtonBase::BindTriggeringInputActionToClick()
 		BindArgs.bForceHold = GetConvertInputActionToHold();
 
 		BindArgs.InputMode = InputModeOverride;
-		
+
 		TriggeringBindingHandle = RegisterUIActionBinding(BindArgs);
 	}
 }
@@ -626,7 +667,7 @@ void UCommonButtonBase::UnbindTriggeringInputActionToClick()
 	{
 		TriggeringBindingHandle.Unregister();
 	}
-	
+
 	CurrentHoldTime = 0.f;
 	CurrentHoldProgress = 0.f;
 }
@@ -693,19 +734,23 @@ void UCommonButtonBase::SetIsInteractionEnabled(bool bInIsInteractionEnabled)
 
 	UpdateInputActionWidgetVisibility();
 
-	// If the hover state changed due to an interactability change, trigger internal logic accordingly.
-	const bool bIsHoveredNow = IsHovered();
-	if (bWasHovered != bIsHoveredNow)
+	if (ShouldProcessHoverEvent(EHoverEventSource::InteractabilityChanged))
 	{
-		if (bIsHoveredNow)
+		// If the hover state changed due to an interactability change, trigger internal logic accordingly.
+		const bool bIsHoveredNow = IsHovered();
+		if (bWasHovered != bIsHoveredNow)
 		{
-			NativeOnHovered();
-		}
-		else
-		{
-			NativeOnUnhovered();
+			if (bIsHoveredNow)
+			{
+				NativeOnHovered();
+			}
+			else
+			{
+				NativeOnUnhovered();
+			}
 		}
 	}
+
 	SetButtonStyle();
 }
 
@@ -812,7 +857,7 @@ bool UCommonButtonBase::NativeOnHoldProgress(float DeltaTime)
 			HoldReset();
 			return false;
 		}
-		
+
 		return true;
 	}
 	HoldReset();
@@ -831,14 +876,14 @@ bool UCommonButtonBase::NativeOnHoldProgressRollback(float DeltaTime)
 		{
 			FTSTicker::GetCoreTicker().RemoveTicker(HoldProgressRollbackTickerHandle);
 			HoldProgressRollbackTickerHandle = nullptr;
-		
+
 			return false;
 		}
-		
+
 		return true;
 	}
 	HoldReset();
-	
+
 	return false;
 }
 
@@ -908,17 +953,20 @@ void UCommonButtonBase::SetIsSelected(bool InSelected, bool bGiveClickFeedback)
 		}
 	}
 
-	// If the hover state changed due to a selection change, trigger internal logic accordingly.
-	const bool bIsHoveredNow = IsHovered();
-	if (bWasHovered != bIsHoveredNow)
+	if (ShouldProcessHoverEvent(EHoverEventSource::SelectionChanged))
 	{
-		if (bIsHoveredNow)
+		// If the hover state changed due to a selection change, trigger internal logic accordingly.
+		const bool bIsHoveredNow = IsHovered();
+		if (bWasHovered != bIsHoveredNow)
 		{
-			NativeOnHovered();
-		}
-		else
-		{
-			NativeOnUnhovered();
+			if (bIsHoveredNow)
+			{
+				NativeOnHovered();
+			}
+			else
+			{
+				NativeOnUnhovered();
+			}
 		}
 	}
 }
@@ -967,7 +1015,7 @@ void UCommonButtonBase::SetSelectedInternal(bool bInSelected, bool bAllowSound /
 	{
 		// Once deselected, restore the root button interactivity to the desired state
 		RootButton->SetInteractionEnabled(bInteractionEnabled);
-		
+
 		NativeOnDeselected(bBroadcast);
 	}
 
@@ -997,7 +1045,10 @@ void UCommonButtonBase::NativeOnMouseEnter(const FGeometry& InGeometry, const FP
 
 		if (GetIsEnabled() && bInteractionEnabled)
 		{
-			NativeOnHovered();
+			if (ShouldProcessHoverEvent(EHoverEventSource::MouseEvent))
+			{
+				NativeOnHovered();
+			}
 		}
 	}
 }
@@ -1010,7 +1061,10 @@ void UCommonButtonBase::NativeOnMouseLeave(const FPointerEvent& InMouseEvent)
 
 		if (GetIsEnabled() && bInteractionEnabled)
 		{
-			NativeOnUnhovered();
+			if (ShouldProcessHoverEvent(EHoverEventSource::MouseEvent))
+			{
+				NativeOnUnhovered();
+			}
 		}
 	}
 }
@@ -1289,11 +1343,11 @@ void UCommonButtonBase::HandleButtonClicked()
 	{
 		// @TODO: Current click rejection method relies on click hold time, this can be refined. See NativeOnHoldProgress.
 		// Also gamepad can indirectly trigger this method, so don't guard against pressed
-    	if (bRequiresHold && CurrentHoldProgress < 1.f)
-    	{
-    		return;
-    	}
-    		
+		if (bRequiresHold && CurrentHoldProgress < 1.f)
+		{
+			return;
+		}
+
 		if (bTriggerClickedAfterSelection)
 		{
 			SetIsSelected(!bSelected, false);
@@ -1338,10 +1392,9 @@ void UCommonButtonBase::HandleButtonPressed()
 	NativeOnPressed();
 
 	UCommonInputSubsystem* CommonInputSubsystem = GetInputSubsystem();
-
-	if (CommonInputSubsystem && CommonInputSubsystem->GetCurrentInputType() == ECommonInputType::Touch && bSimulateHoverOnTouchInput && UE::CommonUI::Private::bEnableSimulateHoverOnTouchInput)
+	
+	if (ShouldProcessHoverEvent(EHoverEventSource::SimulationForTouch))
 	{
-		// Simulate hover events when using touch input
 		NativeOnHovered();
 	}
 
@@ -1366,8 +1419,8 @@ void UCommonButtonBase::HandleButtonReleased()
 	NativeOnReleased();
 
 	UCommonInputSubsystem* CommonInputSubsystem = GetInputSubsystem();
-
-	if (CommonInputSubsystem && CommonInputSubsystem->GetCurrentInputType() == ECommonInputType::Touch && bSimulateHoverOnTouchInput && UE::CommonUI::Private::bEnableSimulateHoverOnTouchInput)
+	
+	if (ShouldProcessHoverEvent(EHoverEventSource::SimulationForTouch))
 	{
 		// Simulate hover events when using touch input
 		NativeOnUnhovered();
@@ -1383,7 +1436,7 @@ void UCommonButtonBase::HandleButtonReleased()
 		{
 			// Begin hold progress rollback
 			HoldProgressRollbackTickerHandle = FTSTicker::GetCoreTicker().AddTicker(FTickerDelegate::CreateUObject(this, &UCommonButtonBase::NativeOnHoldProgressRollback));
-            		
+
 			FTSTicker::GetCoreTicker().RemoveTicker(HoldTickerHandle);
 			HoldTickerHandle = nullptr;
 		}
@@ -1393,7 +1446,7 @@ void UCommonButtonBase::HandleButtonReleased()
 FReply UCommonButtonBase::NativeOnFocusReceived(const FGeometry& InGeometry, const FFocusEvent& InFocusEvent)
 {
 	FReply Reply = Super::NativeOnFocusReceived(InGeometry, InFocusEvent);
-	
+
 	HandleFocusReceived();
 
 	return Reply;
@@ -1409,7 +1462,7 @@ void UCommonButtonBase::NativeOnFocusLost(const FFocusEvent& InFocusEvent)
 void UCommonButtonBase::NativeOnSelected(bool bBroadcast)
 {
 	BP_OnSelected();
-	
+
 	if (bBroadcast)
 	{
 		OnIsSelectedChanged().Broadcast(true);
@@ -1432,9 +1485,14 @@ void UCommonButtonBase::NativeOnDeselected(bool bBroadcast)
 
 void UCommonButtonBase::NativeOnHovered()
 {
+	if (!ShouldProcessHoverEvent(EHoverEventSource::Unknown))
+	{
+		return;
+	}
+
 	BP_OnHovered();
 	OnHovered().Broadcast();
-	
+
 	if (OnButtonBaseHovered.IsBound())
 	{
 		OnButtonBaseHovered.Broadcast(this);
@@ -1450,9 +1508,14 @@ void UCommonButtonBase::NativeOnHovered()
 
 void UCommonButtonBase::NativeOnUnhovered()
 {
+	if (!ShouldProcessHoverEvent(EHoverEventSource::Unknown))
+	{
+		return;
+	}
+
 	BP_OnUnhovered();
 	OnUnhovered().Broadcast();
-	
+
 	if (OnButtonBaseUnhovered.IsBound())
 	{
 		OnButtonBaseUnhovered.Broadcast(this);
@@ -1477,7 +1540,7 @@ void UCommonButtonBase::NativeOnClicked()
 		{
 			OnButtonBaseClicked.Broadcast(this);
 		}
-		
+
 		FString ButtonName, ABTestName, ExtraData;
 		if (GetButtonAnalyticInfo(ButtonName, ABTestName, ExtraData))
 		{
@@ -1628,11 +1691,11 @@ void UCommonButtonBase::BuildStyles()
 		else
 		{
 			SelectedStyle.PressedSlateSound =
-			bHasPressedSlateSoundOverride || !CommonButtonStyle->SelectedPressedSlateSound
-			? NormalStyle.PressedSlateSound
-			: CommonButtonStyle->SelectedPressedSlateSound.Sound;
+				bHasPressedSlateSoundOverride || !CommonButtonStyle->SelectedPressedSlateSound
+				? NormalStyle.PressedSlateSound
+				: CommonButtonStyle->SelectedPressedSlateSound.Sound;
 		}
-		
+
 		if (SelectedHoveredSlateSoundOverride.GetResourceObject())
 		{
 			SelectedStyle.HoveredSlateSound = SelectedHoveredSlateSoundOverride;
@@ -1640,9 +1703,9 @@ void UCommonButtonBase::BuildStyles()
 		else
 		{
 			SelectedStyle.HoveredSlateSound =
-			bHasHoveredSlateSoundOverride || !CommonButtonStyle->SelectedHoveredSlateSound
-			? NormalStyle.HoveredSlateSound
-			: CommonButtonStyle->SelectedHoveredSlateSound.Sound;
+				bHasHoveredSlateSoundOverride || !CommonButtonStyle->SelectedHoveredSlateSound
+				? NormalStyle.HoveredSlateSound
+				: CommonButtonStyle->SelectedHoveredSlateSound.Sound;
 		}
 
 		// Locked State Sound overrides
@@ -1650,16 +1713,16 @@ void UCommonButtonBase::BuildStyles()
 		if (CommonButtonStyle->LockedPressedSlateSound || LockedPressedSlateSoundOverride.GetResourceObject())
 		{
 			LockedStyle.PressedSlateSound =
-			LockedPressedSlateSoundOverride.GetResourceObject()
-			? LockedPressedSlateSoundOverride
-			: CommonButtonStyle->LockedPressedSlateSound.Sound;
+				LockedPressedSlateSoundOverride.GetResourceObject()
+				? LockedPressedSlateSoundOverride
+				: CommonButtonStyle->LockedPressedSlateSound.Sound;
 		}
 		if (CommonButtonStyle->LockedHoveredSlateSound || LockedHoveredSlateSoundOverride.GetResourceObject())
 		{
 			LockedStyle.HoveredSlateSound =
-			LockedHoveredSlateSoundOverride.GetResourceObject()
-			? LockedHoveredSlateSoundOverride
-			: CommonButtonStyle->LockedHoveredSlateSound.Sound;
+				LockedHoveredSlateSoundOverride.GetResourceObject()
+				? LockedHoveredSlateSoundOverride
+				: CommonButtonStyle->LockedHoveredSlateSound.Sound;
 		}
 
 		SetButtonStyle();
@@ -1823,7 +1886,7 @@ void UCommonButtonBase::UpdateInputActionWidgetVisibility()
 		bool bHidden = false;
 
 		UCommonInputSubsystem* CommonInputSubsystem = GetInputSubsystem();
-		
+
 		if (bHideInputAction)
 		{
 			bHidden = true;
diff --git a/Engine/Plugins/Runtime/CommonUI/Source/CommonUI/Public/CommonButtonBase.h b/Engine/Plugins/Runtime/CommonUI/Source/CommonUI/Public/CommonButtonBase.h
--- a/Engine/Plugins/Runtime/CommonUI/Source/CommonUI/Public/CommonButtonBase.h
+++ b/Engine/Plugins/Runtime/CommonUI/Source/CommonUI/Public/CommonButtonBase.h
@@ -22,6 +22,15 @@ class SBox;
 class UCommonActionWidget;
 class UInputAction;
 
+enum class EHoverEventSource : uint8
+{
+	Unknown,
+	MouseEvent,
+	InteractabilityChanged,
+	SelectionChanged,
+	SimulationForTouch,
+};
+
 class COMMONUI_API FCommonButtonMetaData : public ISlateMetaData
 {
 public:
@@ -529,6 +538,10 @@ protected:
 	UFUNCTION()
 	COMMONUI_API virtual void OnInputMethodChanged(ECommonInputType CurrentInputType);
 
+	/** This should be used as a check before calling NativeOnHovered and NativeOnUnhovered.
+	returns false if this button is set to explicitly block hover events on touch.*/
+	COMMONUI_API bool ShouldProcessHoverEvent(EHoverEventSource HoverReason);
+	
 	/** If HoldData is valid, assigns its values to Keyboard and Mouse, Gamepad and Touch, based off the Current Input Type. */
     UFUNCTION()
 	COMMONUI_API virtual void UpdateHoldData(ECommonInputType CurrentInputType);
@@ -797,7 +810,7 @@ protected:
 	TSubclassOf<UCommonUIHoldData> HoldData;
 	
 	/** True if this button should play the hover effect when pressed by a touch input */
-	UPROPERTY(EditAnywhere, BlueprintReadOnly, Category = Input, AdvancedDisplay)
+	UPROPERTY(EditAnywhere, BlueprintReadOnly, Category = Input, AdvancedDisplay, meta = (EditCondition="IsHoverSimulationOnTouchAvailable()", EditConditionHides))
 	bool bSimulateHoverOnTouchInput = true;
 
 private:
@@ -921,6 +934,9 @@ private:
 	//Set this to Game for special cases where an input action needs to be set for an in-game button.
 	UPROPERTY(EditAnywhere, Category = Input, AdvancedDisplay)
 	ECommonInputMode InputModeOverride = ECommonInputMode::Menu;
+
+	UFUNCTION()
+	bool IsHoverSimulationOnTouchAvailable() const;
 	
 	void BuildStyles();
 	void SetButtonStyle();
