Optimize Menu. Don't create all tabs at once. Only create a tab's content when we actually need it (on selection)

- Added a method of registering content widgets with CommonTabListWidgetBase after that Tab's ID has already been initialized

#changelist validated

#virtualized

diff --git a/Engine/Plugins/Runtime/CommonUI/Source/CommonUI/Private/CommonTabListWidgetBase.cpp b/Engine/Plugins/Runtime/CommonUI/Source/CommonUI/Private/CommonTabListWidgetBase.cpp
--- a/Engine/Plugins/Runtime/CommonUI/Source/CommonUI/Private/CommonTabListWidgetBase.cpp
+++ b/Engine/Plugins/Runtime/CommonUI/Source/CommonUI/Private/CommonTabListWidgetBase.cpp
@@ -157,6 +157,9 @@ bool UCommonTabListWidgetBase::RemoveTab(FName TabNameID)
 
 void UCommonTabListWidgetBase::RemoveAllTabs()
 {
+	// Suppress selection when tearing down tabs until we're down
+	TabButtonGroup->SetSelectionRequired(false);
+
 	TArray<FName> RegisteredTabIDs;
 	RegisteredTabsByID.GetKeys(RegisteredTabIDs);
 	for (const FName& RegisteredTabID : RegisteredTabIDs)
@@ -164,6 +167,9 @@ void UCommonTabListWidgetBase::RemoveAllTabs()
 		// Go through the regular remove tab flow, so that we properly free up pooled widgets
 		RemoveTab(RegisteredTabID);
 	}
+
+	// Restore selection requirement on button group
+	TabButtonGroup->SetSelectionRequired(true);
 }
 
 int32 UCommonTabListWidgetBase::GetTabCount() const
@@ -361,6 +367,57 @@ UCommonButtonBase* UCommonTabListWidgetBase::GetTabButtonBaseByID(FName TabNameI
 	return nullptr;
 }
 
+bool UCommonTabListWidgetBase::HasTabContentWidget(const FName TabNameId) const
+{
+	const FCommonRegisteredTabInfo* FoundTabInfo = RegisteredTabsByID.Find(TabNameId);
+	return FoundTabInfo && FoundTabInfo->ContentInstance != nullptr;
+}
+
+bool UCommonTabListWidgetBase::RegisterTabContentWidget(const FName TabNameId, UWidget* ContentWidget)
+{
+	if (!ensure(ContentWidget))
+	{
+		return false;
+	}
+
+	FCommonRegisteredTabInfo* FoundTabInfo = RegisteredTabsByID.Find(TabNameId);
+	if (!ensure(FoundTabInfo))
+	{
+		return false;
+	}
+
+	UWidget* OldContentWidget = FoundTabInfo->ContentInstance;
+	FoundTabInfo->ContentInstance = ContentWidget;
+
+	if (UCommonAnimatedSwitcher* Switcher = LinkedSwitcher.Get())
+	{
+		// Remove Old Widget if it exists
+		if (OldContentWidget != nullptr)
+		{
+			Switcher->RemoveChild(OldContentWidget);
+		}
+
+		// Add the new widget
+		Switcher->AddChild(ContentWidget);
+
+		// If this tab is selected we need to set it as the active widget
+		if (TabNameId == GetSelectedTabId())
+		{
+			Switcher->SetActiveWidget(ContentWidget);
+		}
+	}
+
+	return true;
+}
+
+void UCommonTabListWidgetBase::SetSelectionRequired(bool bSelectionRequired)
+{
+	if (TabButtonGroup)
+	{
+		TabButtonGroup->SetSelectionRequired(bSelectionRequired);
+	}
+}
+
 void UCommonTabListWidgetBase::NativeOnInitialized()
 {
 	Super::NativeOnInitialized();
@@ -544,4 +601,4 @@ void UCommonTabListWidgetBase::RemoveTab_Internal(const FName TabNameID, const F
 	// Callbacks
 	HandleTabRemoval(TabNameID, TabButton);
 	OnTabButtonRemoval.Broadcast(TabNameID, TabButton);
-}
\ No newline at end of file
+}
diff --git a/Engine/Plugins/Runtime/CommonUI/Source/CommonUI/Public/CommonTabListWidgetBase.h b/Engine/Plugins/Runtime/CommonUI/Source/CommonUI/Public/CommonTabListWidgetBase.h
--- a/Engine/Plugins/Runtime/CommonUI/Source/CommonUI/Public/CommonTabListWidgetBase.h
+++ b/Engine/Plugins/Runtime/CommonUI/Source/CommonUI/Public/CommonTabListWidgetBase.h
@@ -150,6 +150,17 @@ public:
 	UFUNCTION(BlueprintCallable, Category = TabList)
 	UCommonButtonBase* GetTabButtonBaseByID(FName TabNameID) const;
 
+	/** Checks if a tab has an associated content widget */
+	UFUNCTION(BlueprintCallable, Category = "Tab List")
+	bool HasTabContentWidget(const FName TabNameId) const;
+
+	/** Registers a content widget with a previously created tab with ID TabNameId. If a linked switcher has been setup, it will also be added to it */
+	UFUNCTION(BlueprintCallable, Category = "Tab List")
+	bool RegisterTabContentWidget(const FName TabNameId, UWidget* ContentWidget);
+
+	/** Allows one to temporarily disable the selection-required behavior TabButtonGroup, useful during initialization and destruction of a UCommonTabListWidgetBase */
+	void SetSelectionRequired(bool bSelectionRequired);
+
 protected:
 	// UUserWidget interface
 	virtual void NativeOnInitialized() override;
