Fix crash on null button group

diff --git a/Engine/Plugins/Runtime/CommonUI/Source/CommonUI/Private/CommonTabListWidgetBase.cpp b/Engine/Plugins/Runtime/CommonUI/Source/CommonUI/Private/CommonTabListWidgetBase.cpp
--- a/Engine/Plugins/Runtime/CommonUI/Source/CommonUI/Private/CommonTabListWidgetBase.cpp
+++ b/Engine/Plugins/Runtime/CommonUI/Source/CommonUI/Private/CommonTabListWidgetBase.cpp
@@ -158,7 +158,7 @@ bool UCommonTabListWidgetBase::RemoveTab(FName TabNameID)
 void UCommonTabListWidgetBase::RemoveAllTabs()
 {
 	// Suppress selection when tearing down tabs until we're down
-	TabButtonGroup->SetSelectionRequired(false);
+	SetSelectionRequired(false);
 
 	TArray<FName> RegisteredTabIDs;
 	RegisteredTabsByID.GetKeys(RegisteredTabIDs);
@@ -169,7 +169,7 @@ void UCommonTabListWidgetBase::RemoveAllTabs()
 	}
 
 	// Restore selection requirement on button group
-	TabButtonGroup->SetSelectionRequired(true);
+	SetSelectionRequired(true);
 }
 
 int32 UCommonTabListWidgetBase::GetTabCount() const
@@ -424,7 +424,7 @@ void UCommonTabListWidgetBase::NativeOnInitialized()
 
 	// Create the button group once up-front
 	TabButtonGroup = NewObject<UCommonButtonGroupBase>(this);
-	TabButtonGroup->SetSelectionRequired(true);
+	SetSelectionRequired(true);
 	TabButtonGroup->OnSelectedButtonBaseChanged.AddDynamic(this, &UCommonTabListWidgetBase::HandleTabButtonSelected);
 }
 
@@ -552,7 +552,7 @@ void UCommonTabListWidgetBase::RebuildTabList()
 	const FName CachedActiveTabID = ActiveTabID;
 
 	// Disable selection required temporarily so we can deselect everything, rebuild the list, then select the tab we want.
-	TabButtonGroup->SetSelectionRequired(false);
+	SetSelectionRequired(false);
 	TabButtonGroup->DeselectAll();
 
 	// Clear all tabs, releasing their widgets back to the widget pool
@@ -579,7 +579,7 @@ void UCommonTabListWidgetBase::RebuildTabList()
 	SelectTabByID(CachedActiveTabID, bSuppressClickFeedback);
 
 	// Turn back on selection requirement
-	TabButtonGroup->SetSelectionRequired(true);
+	SetSelectionRequired(true);
 	
 	// Broadcast our rebuilt delegate
 	OnTabListRebuilt.Broadcast();
