CommonTabListWidgetBase Logic Bug. Call SetSelectionRequired in NativeDestruct to avoid getting nested calls that were preemptively re-enabling the selection requirement in UCommonTabListWidgetBase::RebuildTabList

diff --git a/Engine/Plugins/Runtime/CommonUI/Source/CommonUI/Private/CommonTabListWidgetBase.cpp b/Engine/Plugins/Runtime/CommonUI/Source/CommonUI/Private/CommonTabListWidgetBase.cpp
--- a/Engine/Plugins/Runtime/CommonUI/Source/CommonUI/Private/CommonTabListWidgetBase.cpp
+++ b/Engine/Plugins/Runtime/CommonUI/Source/CommonUI/Private/CommonTabListWidgetBase.cpp
@@ -157,9 +157,6 @@ bool UCommonTabListWidgetBase::RemoveTab(FName TabNameID)
 
 void UCommonTabListWidgetBase::RemoveAllTabs()
 {
-	// Suppress selection when tearing down tabs until we're down
-	SetSelectionRequired(false);
-
 	TArray<FName> RegisteredTabIDs;
 	RegisteredTabsByID.GetKeys(RegisteredTabIDs);
 	for (const FName& RegisteredTabID : RegisteredTabIDs)
@@ -167,9 +164,6 @@ void UCommonTabListWidgetBase::RemoveAllTabs()
 		// Go through the regular remove tab flow, so that we properly free up pooled widgets
 		RemoveTab(RegisteredTabID);
 	}
-
-	// Restore selection requirement on button group
-	SetSelectionRequired(true);
 }
 
 int32 UCommonTabListWidgetBase::GetTabCount() const
@@ -445,7 +439,12 @@ void UCommonTabListWidgetBase::NativeDestruct()
 	SetListeningForInput(false);
 
 	ActiveTabID = NAME_None;
+
+	// Suppress selection when tearing down tabs
+	SetSelectionRequired(false);
 	RemoveAllTabs();
+	SetSelectionRequired(true);
+
 	if (TabButtonGroup)
 	{
 		TabButtonGroup->RemoveAll();
