CommonUI: Add option for action bar to display actions without valid keys for the current input type

This is useful for cases where the action bar is serving as a list of clickable buttons

#rnx

diff --git a/Engine/Plugins/Runtime/CommonUI/Source/CommonUI/Private/CommonActivatableWidget.cpp b/Engine/Plugins/Runtime/CommonUI/Source/CommonUI/Private/CommonActivatableWidget.cpp
--- a/Engine/Plugins/Runtime/CommonUI/Source/CommonUI/Private/CommonActivatableWidget.cpp
+++ b/Engine/Plugins/Runtime/CommonUI/Source/CommonUI/Private/CommonActivatableWidget.cpp
@@ -38,6 +38,8 @@ void UCommonActivatableWidget::NativeConstruct()
 			FBindUIActionArgs BindArgs(ICommonInputModule::GetSettings().GetEnhancedInputBackAction(), FSimpleDelegate::CreateUObject(this, &UCommonActivatableWidget::HandleBackAction));
 			BindArgs.bDisplayInActionBar = bIsBackActionDisplayedInActionBar;
 			BindArgs.OverrideDisplayName = OverrideBackActionDisplayName;
+			BindArgs.InputTypesExemptFromValidKeyCheck.Add(ECommonInputType::MouseAndKeyboard);
+			BindArgs.InputTypesExemptFromValidKeyCheck.Add(ECommonInputType::Touch);
 
 			DefaultBackActionHandle = RegisterUIActionBinding(BindArgs);
 		}
@@ -46,6 +48,8 @@ void UCommonActivatableWidget::NativeConstruct()
 			FBindUIActionArgs BindArgs(ICommonInputModule::GetSettings().GetDefaultBackAction(), FSimpleDelegate::CreateUObject(this, &UCommonActivatableWidget::HandleBackAction));
 			BindArgs.bDisplayInActionBar = bIsBackActionDisplayedInActionBar;
 			BindArgs.OverrideDisplayName = OverrideBackActionDisplayName;
+			BindArgs.InputTypesExemptFromValidKeyCheck.Add(ECommonInputType::MouseAndKeyboard);
+			BindArgs.InputTypesExemptFromValidKeyCheck.Add(ECommonInputType::Touch);
 
 			DefaultBackActionHandle = RegisterUIActionBinding(BindArgs);
 		}
diff --git a/Engine/Plugins/Runtime/CommonUI/Source/CommonUI/Private/Input/CommonBoundActionBar.cpp b/Engine/Plugins/Runtime/CommonUI/Source/CommonUI/Private/Input/CommonBoundActionBar.cpp
--- a/Engine/Plugins/Runtime/CommonUI/Source/CommonUI/Private/Input/CommonBoundActionBar.cpp
+++ b/Engine/Plugins/Runtime/CommonUI/Source/CommonUI/Private/Input/CommonBoundActionBar.cpp
@@ -175,7 +175,8 @@ void UCommonBoundActionBar::HandleDeferredDisplayUpdate()
 									return false;
 								}
 
-								const bool bIsValidEnhancedInputAction = CommonUI::IsEnhancedInputSupportEnabled() && CommonUI::ActionValidForInputType(ActionRouter->GetLocalPlayer(), PlayerInputType, Binding->InputAction.Get());
+								const bool bIsValidEnhancedInputAction = Binding->InputTypesExemptFromValidKeyCheck.Contains(PlayerInputType) ||
+									(CommonUI::IsEnhancedInputSupportEnabled() && CommonUI::ActionValidForInputType(ActionRouter->GetLocalPlayer(), PlayerInputType, Binding->InputAction.Get()));
 								if (!bIsValidEnhancedInputAction)
 								{
 									const bool bIsValidDataTableInputAction = CommonUI::ActionValidForInputType(ActionRouter->GetLocalPlayer(), PlayerInputType, Binding->GetLegacyInputActionData());
diff --git a/Engine/Plugins/Runtime/CommonUI/Source/CommonUI/Private/Input/UIActionRouterTypes.cpp b/Engine/Plugins/Runtime/CommonUI/Source/CommonUI/Private/Input/UIActionRouterTypes.cpp
--- a/Engine/Plugins/Runtime/CommonUI/Source/CommonUI/Private/Input/UIActionRouterTypes.cpp
+++ b/Engine/Plugins/Runtime/CommonUI/Source/CommonUI/Private/Input/UIActionRouterTypes.cpp
@@ -114,6 +114,7 @@ FUIActionBinding::FUIActionBinding(const UWidget& InBoundWidget, const FBindUIAc
 	, BoundWidget(&InBoundWidget)
 	, InputMode(BindArgs.InputMode)
 	, bDisplayInActionBar(BindArgs.bDisplayInActionBar)
+	, InputTypesExemptFromValidKeyCheck(BindArgs.InputTypesExemptFromValidKeyCheck)
 	, ActionDisplayName(BindArgs.OverrideDisplayName)
 	, OnExecuteAction(BindArgs.OnExecuteAction)
 	, Handle(IdCounter++)
diff --git a/Engine/Plugins/Runtime/CommonUI/Source/CommonUI/Public/Input/CommonUIInputTypes.h b/Engine/Plugins/Runtime/CommonUI/Source/CommonUI/Public/Input/CommonUIInputTypes.h
--- a/Engine/Plugins/Runtime/CommonUI/Source/CommonUI/Public/Input/CommonUIInputTypes.h
+++ b/Engine/Plugins/Runtime/CommonUI/Source/CommonUI/Public/Input/CommonUIInputTypes.h
@@ -8,6 +8,8 @@
 #include "Engine/EngineBaseTypes.h"
 #include "InputAction.h"
 
+enum class ECommonInputType : uint8;
+
 struct COMMONUI_API FBindUIActionArgs
 {
 	FBindUIActionArgs(FUIActionTag InActionTag, const FSimpleDelegate& InOnExecuteAction)
@@ -59,6 +61,12 @@ struct COMMONUI_API FBindUIActionArgs
 	ECommonInputMode InputMode = ECommonInputMode::Menu;
 	EInputEvent KeyEvent = IE_Pressed;
 
+	/**
+	 * By default, the action bar only displays prompts for actions with keys valid for the current input type, any input types added here will
+	 * skip that check and display this action regardless of the bound keys
+	 */
+	TSet<ECommonInputType> InputTypesExemptFromValidKeyCheck;
+
 	/**
 	 * A persistent binding is always registered and will be executed regardless of the activation status of the binding widget's parentage.
 	 * Persistent bindings also never stomp one another - if two are bound to the same action, both will execute. Use should be kept to a minimum.
diff --git a/Engine/Plugins/Runtime/CommonUI/Source/CommonUI/Public/Input/UIActionBinding.h b/Engine/Plugins/Runtime/CommonUI/Source/CommonUI/Public/Input/UIActionBinding.h
--- a/Engine/Plugins/Runtime/CommonUI/Source/CommonUI/Public/Input/UIActionBinding.h
+++ b/Engine/Plugins/Runtime/CommonUI/Source/CommonUI/Public/Input/UIActionBinding.h
@@ -7,6 +7,7 @@
 #include "Input/CommonUIInputSettings.h" // IWYU pragma: keep
 #include "Input/UIActionBindingHandle.h"
 
+enum class ECommonInputType : uint8;
 enum EInputEvent : int;
 struct FKey;
 struct FUIActionKeyMapping;
@@ -66,6 +67,7 @@ struct COMMONUI_API FUIActionBinding
 	int32 UserIndex = INDEX_NONE;
 
 	bool bDisplayInActionBar = false;
+	TSet<ECommonInputType> InputTypesExemptFromValidKeyCheck;
 	FText ActionDisplayName;
 	
 	TWeakPtr<FActionRouterBindingCollection> OwningCollection;
