Delay rebuilding of action bar buttons until no existing buttons any longer have mouse capture
Relink focus state to discovery browser horizontal list item

#tests PIE, CLIENT Discovery Broswer
#rnx

diff --git a/Engine/Plugins/Runtime/CommonUI/Source/CommonUI/Private/Input/CommonBoundActionBar.cpp b/Engine/Plugins/Runtime/CommonUI/Source/CommonUI/Private/Input/CommonBoundActionBar.cpp
--- a/Engine/Plugins/Runtime/CommonUI/Source/CommonUI/Private/Input/CommonBoundActionBar.cpp
+++ b/Engine/Plugins/Runtime/CommonUI/Source/CommonUI/Private/Input/CommonBoundActionBar.cpp
@@ -9,10 +9,12 @@
 #include "EnhancedInputSubsystems.h"
 #include "Engine/GameInstance.h"
 #include "Engine/GameViewportClient.h"
+#include "Framework/Application/SlateUser.h"
 #include "InputAction.h"
 #include "Input/CommonBoundActionButtonInterface.h"
 #include "Input/CommonUIActionRouterBase.h"
 #include "Input/UIActionBinding.h"
+#include "Layout/WidgetPath.h"
 
 #include UE_INLINE_GENERATED_CPP_BY_NAME(CommonBoundActionBar)
 
@@ -127,12 +129,21 @@ void UCommonBoundActionBar::UpdateDisplay()
 
 		FTSTicker::GetCoreTicker().AddTicker(FTickerDelegate::CreateWeakLambda(this, [this](float InDelta)
 			{
-				HandleDeferredDisplayUpdate();
-				return false;
+				if (IsSafeToUpdateDisplay())
+				{
+					HandleDeferredDisplayUpdate();
+					return false;
+				}
+				return true;
 			}));
 	}
 }
 
+bool UCommonBoundActionBar::IsSafeToUpdateDisplay() const
+{
+	return !DoAnyActionButtonsHaveMouseCapture();
+}
+
 void UCommonBoundActionBar::HandleDeferredDisplayUpdate()
 {
 	ActionBarUpdateBegin();
@@ -340,16 +351,7 @@ void UCommonBoundActionBar::HandlePlayerRemoved(int32 PlayerIdx)
 
 void UCommonBoundActionBar::HandledInputTypeUpdated(ECommonInputType InputType)
 {
-	if (!bIsRefreshQueued)
-	{
-		bIsRefreshQueued = true;
-
-		FTSTicker::GetCoreTicker().AddTicker(FTickerDelegate::CreateWeakLambda(this, [this](float InDelta)
-		{
-			HandleDeferredDisplayUpdate();
-			return false;
-		}));
-	}
+	UpdateDisplay();
 }
 
 void UCommonBoundActionBar::MonitorPlayerActions(const ULocalPlayer* NewPlayer)
@@ -395,4 +397,27 @@ void UCommonBoundActionBar::ActionBarUpdateEnd()
 	ActionBarUpdateEndImpl();
 }
 
+bool UCommonBoundActionBar::DoAnyActionButtonsHaveMouseCapture() const
+{
+	if (const ULocalPlayer* LocalPlayer = GetOwningLocalPlayer())
+	{
+		if (TSharedPtr<const FSlateUser> SlateUser = LocalPlayer->GetSlateUser())
+		{
+			if (SlateUser->HasAnyCapture())
+			{
+				const FWeakWidgetPath CapturePath = SlateUser->GetWeakCursorCapturePath();
+				for (const UUserWidget* Element : GetAllEntries())
+				{
+					if (CapturePath.ContainsWidget(Element->GetCachedWidget().Get()))
+					{
+						return true;
+					}
+				}
+			}
+		}
+	}
+
+	return false;
+}
+
 #undef LOCTEXT_NAMESPACE
diff --git a/Engine/Plugins/Runtime/CommonUI/Source/CommonUI/Public/Input/CommonBoundActionBar.h b/Engine/Plugins/Runtime/CommonUI/Source/CommonUI/Public/Input/CommonBoundActionBar.h
--- a/Engine/Plugins/Runtime/CommonUI/Source/CommonUI/Public/Input/CommonBoundActionBar.h
+++ b/Engine/Plugins/Runtime/CommonUI/Source/CommonUI/Public/Input/CommonBoundActionBar.h
@@ -51,6 +51,7 @@ private:
 	UFUNCTION()
 	void HandleInputMappingsRebuiltUpdated();
 	void UpdateDisplay();
+	bool IsSafeToUpdateDisplay() const;
 	void HandleDeferredDisplayUpdate();
 	void HandlePlayerAdded(int32 PlayerIdx);
 	void HandlePlayerRemoved(int32 PlayerIdx);
@@ -59,6 +60,7 @@ private:
 
 	void ActionBarUpdateBegin();
 	void ActionBarUpdateEnd();
+	bool DoAnyActionButtonsHaveMouseCapture() const;
 
 	UPROPERTY(EditAnywhere, Category = EntryLayout, meta=(MustImplement = "/Script/CommonUI.CommonBoundActionButtonInterface"))
 	TSubclassOf<UCommonButtonBase> ActionButtonClass;
