Detect touch input type in mouse wheel handling.

This fixes an issue for when the FGestureRecognizer simulates a mouse wheel event so that common ui treats it correctly.

#rnx

diff --git a/Engine/Plugins/Runtime/CommonUI/Source/CommonInput/Private/CommonInputPreprocessor.cpp b/Engine/Plugins/Runtime/CommonUI/Source/CommonInput/Private/CommonInputPreprocessor.cpp
--- a/Engine/Plugins/Runtime/CommonUI/Source/CommonInput/Private/CommonInputPreprocessor.cpp
+++ b/Engine/Plugins/Runtime/CommonUI/Source/CommonInput/Private/CommonInputPreprocessor.cpp
@@ -119,7 +119,7 @@ bool FCommonInputPreprocessor::HandleMouseButtonDoubleClickEvent(FSlateApplicati
 
 bool FCommonInputPreprocessor::HandleMouseWheelOrGestureEvent(FSlateApplication& SlateApp, const FPointerEvent& InWheelEvent, const FPointerEvent* InGestureEvent)
 {
-	const ECommonInputType InputType = GetInputType(InWheelEvent);
+	const ECommonInputType InputType = InGestureEvent ? GetInputType(*InGestureEvent) : GetInputType(InWheelEvent);
 	if (IsRelevantInput(SlateApp, InWheelEvent, InputType))
 	{
 		if (IsInputMethodBlocked(InputType))
@@ -271,7 +271,7 @@ ECommonInputType FCommonInputPreprocessor::GetInputType(const FPointerEvent& Poi
 	const ULocalPlayer& LocalPlayer = *InputSubsystem.GetLocalPlayerChecked();
 #endif
 
-	if (PointerEvent.IsTouchEvent())
+	if (PointerEvent.IsTouchEvent() || PointerEvent.GetGestureType() != EGestureEvent::None)
 	{
 		return ECommonInputType::Touch;
 	}
