CommonUI: Consider the GestureEvent when it's valid in FCommonInputPreprocessor::HandleMouseWheelOrGestureEvent

Fixes touch gestures setting the current input method to MouseAndKeyboard

#rnx

diff --git a/Engine/Plugins/Runtime/CommonUI/Source/CommonInput/Private/CommonInputPreprocessor.cpp b/Engine/Plugins/Runtime/CommonUI/Source/CommonInput/Private/CommonInputPreprocessor.cpp
--- a/Engine/Plugins/Runtime/CommonUI/Source/CommonInput/Private/CommonInputPreprocessor.cpp
+++ b/Engine/Plugins/Runtime/CommonUI/Source/CommonInput/Private/CommonInputPreprocessor.cpp
@@ -119,8 +119,9 @@ bool FCommonInputPreprocessor::HandleMouseButtonDoubleClickEvent(FSlateApplicati
 
 bool FCommonInputPreprocessor::HandleMouseWheelOrGestureEvent(FSlateApplication& SlateApp, const FPointerEvent& InWheelEvent, const FPointerEvent* InGestureEvent)
 {
-	const ECommonInputType InputType = InGestureEvent ? GetInputType(*InGestureEvent) : GetInputType(InWheelEvent);
-	if (IsRelevantInput(SlateApp, InWheelEvent, InputType))
+	const FPointerEvent& EventToConsider = InGestureEvent ? *InGestureEvent : InWheelEvent;
+	const ECommonInputType InputType = GetInputType(EventToConsider);
+	if (IsRelevantInput(SlateApp, EventToConsider, InputType))
 	{
 		if (IsInputMethodBlocked(InputType))
 		{
@@ -271,7 +272,7 @@ ECommonInputType FCommonInputPreprocessor::GetInputType(const FPointerEvent& Poi
 	const ULocalPlayer& LocalPlayer = *InputSubsystem.GetLocalPlayerChecked();
 #endif
 
-	if (PointerEvent.IsTouchEvent() || PointerEvent.GetGestureType() != EGestureEvent::None)
+	if (PointerEvent.IsTouchEvent())
 	{
 		return ECommonInputType::Touch;
 	}
