CommonUI: Add ResetCalculatedActionDomainCache() for widgets to call when their Action Domain is changed at runtime

#rnx

diff --git a/Engine/Plugins/Runtime/CommonUI/Source/CommonUI/Private/CommonActivatableWidget.cpp b/Engine/Plugins/Runtime/CommonUI/Source/CommonUI/Private/CommonActivatableWidget.cpp
--- a/Engine/Plugins/Runtime/CommonUI/Source/CommonUI/Private/CommonActivatableWidget.cpp
+++ b/Engine/Plugins/Runtime/CommonUI/Source/CommonUI/Private/CommonActivatableWidget.cpp
@@ -248,6 +248,11 @@ TObjectPtr<UCommonInputActionDomain> UCommonActivatableWidget::GetCalculatedActi
 	return nullptr;
 }
 
+void UCommonActivatableWidget::ResetCalculatedActionDomainCache()
+{
+	CalculatedActionDomainCache.Reset();
+}
+
 TSharedRef<SWidget> UCommonActivatableWidget::RebuildWidget()
 {
 	// Note: the scoped builder guards against design-time so we don't need to here (as it'd make the scoped lifetime more awkward to leverage)
diff --git a/Engine/Plugins/Runtime/CommonUI/Source/CommonUI/Public/CommonActivatableWidget.h b/Engine/Plugins/Runtime/CommonUI/Source/CommonUI/Public/CommonActivatableWidget.h
--- a/Engine/Plugins/Runtime/CommonUI/Source/CommonUI/Public/CommonActivatableWidget.h
+++ b/Engine/Plugins/Runtime/CommonUI/Source/CommonUI/Public/CommonActivatableWidget.h
@@ -118,6 +118,12 @@ public:
 	 */
 	TObjectPtr<UCommonInputActionDomain> GetCalculatedActionDomain();
 
+	/**
+	 * Reset the Action Domain cached by GetCalculatedActionDomain()
+	 * Call this when updating this widget or a parent widget's action domain
+	 */
+	void ResetCalculatedActionDomainCache();
+
 protected:
 	virtual TSharedRef<SWidget> RebuildWidget() override;
 	virtual void OnWidgetRebuilt() override;
