Remove Ensure related to missing content in UCommonTabListWidgetBase::HandleTabButtonSelected, as we now support this case (content hasn't yet been created)

diff --git a/Engine/Plugins/Runtime/CommonUI/Source/CommonUI/Private/CommonTabListWidgetBase.cpp b/Engine/Plugins/Runtime/CommonUI/Source/CommonUI/Private/CommonTabListWidgetBase.cpp
--- a/Engine/Plugins/Runtime/CommonUI/Source/CommonUI/Private/CommonTabListWidgetBase.cpp
+++ b/Engine/Plugins/Runtime/CommonUI/Source/CommonUI/Private/CommonTabListWidgetBase.cpp
@@ -491,14 +491,10 @@ void UCommonTabListWidgetBase::HandleTabButtonSelected(UCommonButtonBase* Select
 		{
 			ActiveTabID = TabPair.Key;
 
-			if (TabInfo.ContentInstance || LinkedSwitcher.IsValid())
+			if (TabInfo.ContentInstance && LinkedSwitcher.IsValid())
 			{
-				if (ensureMsgf(TabInfo.ContentInstance, TEXT("A CommonTabListWidget tab button lacks a tab content widget to set its linked switcher to.")) &&
-					ensureMsgf(LinkedSwitcher.IsValid(), TEXT("A CommonTabListWidgetBase.has a registered tab with a content widget to switch to, but has no linked activatable widget switcher. Did you forget to call SetLinkedSwitcher to establish the association?")))
-				{
-					// There's already an instance of the widget to display, so go for it
-					LinkedSwitcher->SetActiveWidget(TabInfo.ContentInstance);
-				}
+				// There's already an instance of the widget to display, so go for it
+				LinkedSwitcher->SetActiveWidget(TabInfo.ContentInstance);
 			}
 
 			OnTabSelected.Broadcast(TabPair.Key);
