CommonUI: Default InputTypesExemptFromValidKeyCheck to MouseAndKeyboard and Touch since those input types allow clicking/tapping buttons directly

#rnx

diff --git a/Engine/Plugins/Runtime/CommonUI/Source/CommonUI/Private/CommonActivatableWidget.cpp b/Engine/Plugins/Runtime/CommonUI/Source/CommonUI/Private/CommonActivatableWidget.cpp
--- a/Engine/Plugins/Runtime/CommonUI/Source/CommonUI/Private/CommonActivatableWidget.cpp
+++ b/Engine/Plugins/Runtime/CommonUI/Source/CommonUI/Private/CommonActivatableWidget.cpp
@@ -38,8 +38,6 @@ void UCommonActivatableWidget::NativeConstruct()
 			FBindUIActionArgs BindArgs(ICommonInputModule::GetSettings().GetEnhancedInputBackAction(), FSimpleDelegate::CreateUObject(this, &UCommonActivatableWidget::HandleBackAction));
 			BindArgs.bDisplayInActionBar = bIsBackActionDisplayedInActionBar;
 			BindArgs.OverrideDisplayName = OverrideBackActionDisplayName;
-			BindArgs.InputTypesExemptFromValidKeyCheck.Add(ECommonInputType::MouseAndKeyboard);
-			BindArgs.InputTypesExemptFromValidKeyCheck.Add(ECommonInputType::Touch);
 
 			DefaultBackActionHandle = RegisterUIActionBinding(BindArgs);
 		}
@@ -48,8 +46,6 @@ void UCommonActivatableWidget::NativeConstruct()
 			FBindUIActionArgs BindArgs(ICommonInputModule::GetSettings().GetDefaultBackAction(), FSimpleDelegate::CreateUObject(this, &UCommonActivatableWidget::HandleBackAction));
 			BindArgs.bDisplayInActionBar = bIsBackActionDisplayedInActionBar;
 			BindArgs.OverrideDisplayName = OverrideBackActionDisplayName;
-			BindArgs.InputTypesExemptFromValidKeyCheck.Add(ECommonInputType::MouseAndKeyboard);
-			BindArgs.InputTypesExemptFromValidKeyCheck.Add(ECommonInputType::Touch);
 
 			DefaultBackActionHandle = RegisterUIActionBinding(BindArgs);
 		}
diff --git a/Engine/Plugins/Runtime/CommonUI/Source/CommonUI/Public/Input/CommonUIInputTypes.h b/Engine/Plugins/Runtime/CommonUI/Source/CommonUI/Public/Input/CommonUIInputTypes.h
--- a/Engine/Plugins/Runtime/CommonUI/Source/CommonUI/Public/Input/CommonUIInputTypes.h
+++ b/Engine/Plugins/Runtime/CommonUI/Source/CommonUI/Public/Input/CommonUIInputTypes.h
@@ -5,11 +5,10 @@
 #include "Engine/DataTable.h"
 #include "UITag.h"
 #include "CommonInputModeTypes.h"
+#include "CommonInputTypeEnum.h"
 #include "Engine/EngineBaseTypes.h"
 #include "InputAction.h"
 
-enum class ECommonInputType : uint8;
-
 struct COMMONUI_API FBindUIActionArgs
 {
 	FBindUIActionArgs(FUIActionTag InActionTag, const FSimpleDelegate& InOnExecuteAction)
@@ -65,7 +64,7 @@ struct COMMONUI_API FBindUIActionArgs
 	 * By default, the action bar only displays prompts for actions with keys valid for the current input type, any input types added here will
 	 * skip that check and display this action regardless of the bound keys
 	 */
-	TSet<ECommonInputType> InputTypesExemptFromValidKeyCheck;
+	TSet<ECommonInputType> InputTypesExemptFromValidKeyCheck = { ECommonInputType::MouseAndKeyboard, ECommonInputType::Touch };
 
 	/**
 	 * A persistent binding is always registered and will be executed regardless of the activation status of the binding widget's parentage.
