CommonUI: PR #11992: Allow subclasses of UCommonUIVisibilitySubsystem to expand the list of computed tags

#rnx

diff --git a/Engine/Plugins/Runtime/CommonUI/Source/CommonUI/Public/CommonUIVisibilitySubsystem.h b/Engine/Plugins/Runtime/CommonUI/Source/CommonUI/Public/CommonUIVisibilitySubsystem.h
--- a/Engine/Plugins/Runtime/CommonUI/Source/CommonUI/Public/CommonUIVisibilitySubsystem.h
+++ b/Engine/Plugins/Runtime/CommonUI/Source/CommonUI/Public/CommonUIVisibilitySubsystem.h
@@ -56,7 +56,7 @@ public:
 protected:
 	void RefreshVisibilityTags();
 	void OnInputMethodChanged(ECommonInputType CurrentInputType);
-	FGameplayTagContainer ComputeVisibilityTags() const;
+	virtual FGameplayTagContainer ComputeVisibilityTags() const;
 
 private:
 	FGameplayTagContainer ComputedVisibilityTags;
