CommonUI: Specify engine version when ECommonButtonAcceptKeyHandling was added

#rnx

diff --git a/Engine/Plugins/Runtime/CommonUI/Source/CommonUI/Public/CommonUISettings.h b/Engine/Plugins/Runtime/CommonUI/Source/CommonUI/Public/CommonUISettings.h
--- a/Engine/Plugins/Runtime/CommonUI/Source/CommonUI/Public/CommonUISettings.h
+++ b/Engine/Plugins/Runtime/CommonUI/Source/CommonUI/Public/CommonUISettings.h
@@ -16,12 +16,12 @@ COMMONUI_API UE_DECLARE_GAMEPLAY_TAG_EXTERN(TAG_PlatformTrait_PlayInEditor);
 UENUM()
 enum class ECommonButtonAcceptKeyHandling
 {
-	// (Default for existing projects) When a CommonButton is focused, it will ignore the configured keys for slate navigation's Accept action (see FNavigationConfig::KeyActionRules),
+	// (Default for projects created prior to 5.6) When a CommonButton is focused, it will ignore the configured keys for slate navigation's Accept action (see FNavigationConfig::KeyActionRules),
 	// leaving it up to the button's Triggering Input Action to trigger the button's click event.
 	// This allows other objects to bind input actions to the accept action keys while focusing a common button.
 	Ignore,
 
-	// (Default for new projects) When a CommonButton is focused, it will let all inputs events flow through to the underlying SButton, therefore triggering the button's click event when
+	// (Default for new projects as of 5.6) When a CommonButton is focused, it will let all inputs events flow through to the underlying SButton, therefore triggering the button's click event when
 	// pressing the configured keys for slate navigation's Accept action (see FNavigationConfig::KeyActionRules)
 	// This will prevent input actions that are bound to those keys from triggering while a CommonButton is focused.
 	TriggerClick,
