CommonUI: Don't attempt to load ICommonInputModule while shutting down

#rnx

diff --git a/Engine/Plugins/Runtime/CommonUI/Source/CommonUI/Private/CommonUITypes.cpp b/Engine/Plugins/Runtime/CommonUI/Source/CommonUI/Private/CommonUITypes.cpp
--- a/Engine/Plugins/Runtime/CommonUI/Source/CommonUI/Private/CommonUITypes.cpp
+++ b/Engine/Plugins/Runtime/CommonUI/Source/CommonUI/Private/CommonUITypes.cpp
@@ -246,6 +246,11 @@ FSlateBrush CommonUI::GetIconForInputActions(const UCommonInputSubsystem* Common
 
 bool CommonUI::IsEnhancedInputSupportEnabled()
 {
+	if (IsEngineExitRequested())
+	{
+		return false;
+	}
+
 	static bool bEnabled = ICommonInputModule::Get().GetSettings().GetEnableEnhancedInputSupport();
 	return bEnabled;
 }
