CommonUI: Add CommonUI.ShouldVirtualAcceptSimulateMouseButton CVar to modify if virtual accept events should convert to left mouse button events at runtime

#rnx

diff --git a/Engine/Plugins/Runtime/CommonUI/Source/CommonUI/Private/Input/CommonAnalogCursor.cpp b/Engine/Plugins/Runtime/CommonUI/Source/CommonUI/Private/Input/CommonAnalogCursor.cpp
--- a/Engine/Plugins/Runtime/CommonUI/Source/CommonUI/Private/Input/CommonAnalogCursor.cpp
+++ b/Engine/Plugins/Runtime/CommonUI/Source/CommonUI/Private/Input/CommonAnalogCursor.cpp
@@ -22,6 +22,11 @@
 const float AnalogScrollUpdatePeriod = 0.1f;
 const float ScrollDeadZone = 0.2f;
 
+static const TAutoConsoleVariable<bool> CVarShouldVirtualAcceptSimulateMouseButton(
+	TEXT("CommonUI.ShouldVirtualAcceptSimulateMouseButton"),
+	true,
+	TEXT("Controls if virtual_accept key events will be converted to left mouse button events."));
+
 bool IsEligibleFakeKeyPointerEvent(const FPointerEvent& PointerEvent)
 {
 	FKey EffectingButton = PointerEvent.GetEffectingButton();
@@ -405,6 +410,11 @@ void FCommonAnalogCursor::ShouldHandleRightAnalog(bool bInShouldHandleRightAnalo
 	bShouldHandleRightAnalog = bInShouldHandleRightAnalog;
 }
 
+bool FCommonAnalogCursor::ShouldVirtualAcceptSimulateMouseButton(const FKeyEvent& InKeyEvent, EInputEvent InputEvent) const
+{
+	return CVarShouldVirtualAcceptSimulateMouseButton.GetValueOnGameThread();
+}
+
 //void FCommonAnalogCursor::SetCursorMovementStick(EAnalogStick InCursorMovementStick)
 //{
 //	const EAnalogStick NewStick = InCursorMovementStick == EAnalogStick::Max ? EAnalogStick::Left : InCursorMovementStick;
diff --git a/Engine/Plugins/Runtime/CommonUI/Source/CommonUI/Public/Input/CommonAnalogCursor.h b/Engine/Plugins/Runtime/CommonUI/Source/CommonUI/Public/Input/CommonAnalogCursor.h
--- a/Engine/Plugins/Runtime/CommonUI/Source/CommonUI/Public/Input/CommonAnalogCursor.h
+++ b/Engine/Plugins/Runtime/CommonUI/Source/CommonUI/Public/Input/CommonAnalogCursor.h
@@ -53,7 +53,7 @@ public:
 
 	virtual bool IsAnalogMovementEnabled() const { return bIsAnalogMovementEnabled; }
 
-	virtual bool ShouldVirtualAcceptSimulateMouseButton(const FKeyEvent& InKeyEvent, EInputEvent InputEvent) const { return true; }
+	virtual bool ShouldVirtualAcceptSimulateMouseButton(const FKeyEvent& InKeyEvent, EInputEvent InputEvent) const;
 
 protected:
 	FCommonAnalogCursor(const UCommonUIActionRouterBase& InActionRouter);
