CommonUI: Add CVar to prevent Mouse Wheel and Gestures from updating the current input method

#rnx

diff --git a/Engine/Plugins/Runtime/CommonUI/Source/CommonInput/Private/CommonInputPreprocessor.cpp b/Engine/Plugins/Runtime/CommonUI/Source/CommonInput/Private/CommonInputPreprocessor.cpp
--- a/Engine/Plugins/Runtime/CommonUI/Source/CommonInput/Private/CommonInputPreprocessor.cpp
+++ b/Engine/Plugins/Runtime/CommonUI/Source/CommonInput/Private/CommonInputPreprocessor.cpp
@@ -31,6 +31,11 @@
 #include "Settings/LevelEditorPlaySettings.h"
 #endif
 
+static const TAutoConsoleVariable<bool> CVarShouldMouseWheelAndGestureRefreshInputMethod(
+	TEXT("CommonUI.ShouldMouseWheelAndGestureRefreshInputMethod"),
+	true,
+	TEXT("Controls whether Mouse Wheel and Gesture input events should update the current input method"));
+
 FCommonInputPreprocessor::FCommonInputPreprocessor(UCommonInputSubsystem& InCommonInputSubsystem)
 	: InputSubsystem(InCommonInputSubsystem)
 	, bIgnoreNextMove(false)
@@ -127,7 +132,11 @@ bool FCommonInputPreprocessor::HandleMouseWheelOrGestureEvent(FSlateApplication&
 		{
 			return true;
 		}
-		RefreshCurrentInputMethod(InputType);
+
+		if (CVarShouldMouseWheelAndGestureRefreshInputMethod.GetValueOnGameThread())
+		{
+			RefreshCurrentInputMethod(InputType);
+		}
 	}
 	return false;
 }
