Added BP-accessible focus delegates to CommonButtonBase

diff --git a/Engine/Plugins/Runtime/CommonUI/Source/CommonUI/Private/CommonButtonBase.cpp b/Engine/Plugins/Runtime/CommonUI/Source/CommonUI/Private/CommonButtonBase.cpp
--- a/Engine/Plugins/Runtime/CommonUI/Source/CommonUI/Private/CommonButtonBase.cpp
+++ b/Engine/Plugins/Runtime/CommonUI/Source/CommonUI/Private/CommonButtonBase.cpp
@@ -1379,12 +1379,22 @@ void UCommonButtonBase::HandleFocusReceived()
 	}
 	OnFocusReceived().Broadcast();
 	BP_OnFocusReceived();
+
+	if (OnButtonBaseFocused.IsBound())
+	{
+		OnButtonBaseFocused.Broadcast(this);
+	}
 }
 
 void UCommonButtonBase::HandleFocusLost()
 {
 	OnFocusLost().Broadcast();
 	BP_OnFocusLost();
+
+	if (OnButtonBaseUnfocused.IsBound())
+	{
+		OnButtonBaseUnfocused.Broadcast(this);
+	}
 }
 
 void UCommonButtonBase::HandleButtonPressed()
diff --git a/Engine/Plugins/Runtime/CommonUI/Source/CommonUI/Public/CommonButtonBase.h b/Engine/Plugins/Runtime/CommonUI/Source/CommonUI/Public/CommonButtonBase.h
--- a/Engine/Plugins/Runtime/CommonUI/Source/CommonUI/Public/CommonButtonBase.h
+++ b/Engine/Plugins/Runtime/CommonUI/Source/CommonUI/Public/CommonButtonBase.h
@@ -887,6 +887,12 @@ protected:
 
 	UPROPERTY(BlueprintAssignable, Category = "Events", meta = (AllowPrivateAccess = true, DisplayName = "On Unhovered"))
 	FCommonButtonBaseClicked OnButtonBaseUnhovered;
+
+	UPROPERTY(BlueprintAssignable, Category = "Events", meta = (AllowPrivateAccess = true, DisplayName = "On Focused"))
+	FCommonButtonBaseClicked OnButtonBaseFocused;
+
+	UPROPERTY(BlueprintAssignable, Category = "Events", meta = (AllowPrivateAccess = true, DisplayName = "On Unfocused"))
+	FCommonButtonBaseClicked OnButtonBaseUnfocused;
 	
 	UPROPERTY(BlueprintAssignable, Category = "Events", meta = (AllowPrivateAccess = true, DisplayName = "On Lock Clicked"))
 	FCommonButtonBaseClicked OnButtonBaseLockClicked;
