CommonUI: Add Mechanism for preventing soft-lock when disabling/hidding buttons from Verse.

#rnx

diff --git a/Engine/Plugins/Runtime/CommonUI/Source/CommonUI/Private/CommonButtonBase.cpp b/Engine/Plugins/Runtime/CommonUI/Source/CommonUI/Private/CommonButtonBase.cpp
--- a/Engine/Plugins/Runtime/CommonUI/Source/CommonUI/Private/CommonButtonBase.cpp
+++ b/Engine/Plugins/Runtime/CommonUI/Source/CommonUI/Private/CommonButtonBase.cpp
@@ -16,6 +16,7 @@
 #include "ICommonUIModule.h"
 #include "Widgets/Layout/SBox.h"
 #include "Framework/Application/SlateApplication.h"
+#include "Framework/Application/SlateUser.h"
 #include "Materials/MaterialInstanceDynamic.h"
 #include "Engine/UserInterfaceSettings.h"
 #include "ICommonInputModule.h"
@@ -353,6 +354,7 @@ UCommonButtonBase::UCommonButtonBase(const FObjectInitializer& ObjectInitializer
 	, bSelected(false)
 	, bButtonEnabled(true)
 	, bInteractionEnabled(true)
+	, bNavigateToNextWidgetOnDisable(false)
 	, HoldTime(0.f)
 	, HoldRollbackTime(0.f)
 	, CurrentHoldTime(0.f)
@@ -510,9 +512,25 @@ void UCommonButtonBase::SetIsEnabled(bool bInIsEnabled)
 	{
 		// Note: State is disabled, so we broadcast !bIsEnabled
 		BroadcastBinaryPostStateChange(UWidgetDisabledStateRegistration::Bit, !bInIsEnabled);
+
+		HandleImplicitFocusLost();
 	}
 }
 
+
+void UCommonButtonBase::SetVisibility(ESlateVisibility InVisibility)
+{
+	bool bValueChanged = InVisibility != GetVisibility();
+	
+	Super::SetVisibility(InVisibility);
+
+	if (bValueChanged)
+	{
+		HandleImplicitFocusLost();
+	}
+}
+
+
 bool UCommonButtonBase::NativeIsInteractable() const
 {
 	// If it's enabled, it's "interactable" from a UMG perspective. 
@@ -1991,6 +2009,42 @@ bool UCommonButtonBase::GetIsFocusable() const
 	return IsFocusable();
 }
 
+void UCommonButtonBase::HandleImplicitFocusLost()
+{
+	// Note: This is workaround to avoid users from invalidating focus state
+	// From code, users are able to disable/hide widgets that have the users focus. 
+	// If the widget was disabled, and is interacted with (e.g. Clicked), Slate will attempt to restore user focus. If the widget is hidden, the focus state is lost 
+	// This results in a bad focused widget state as keyboard and controller events will attempt to tunnel through the focus path via FReply::RouteAlongFocusPath
+	// To avoid this, ensure users are focused on enabled widgets by navigating to next available focusable widget if the disabled widget is currently in focus. 
+	if (!bNavigateToNextWidgetOnDisable)
+	{
+		return;
+	}
+
+	if (GetIsEnabled() && IsVisible())
+	{
+		return;
+	}
+
+	TSharedPtr<SWidget> SafeWidget = GetCachedWidget();
+	if (!SafeWidget.IsValid() || !RootButton.IsValid())
+	{
+		return;
+	}
+
+	FSlateApplication::Get().ForEachUser([SafeWidget](FSlateUser& User)
+	{
+		if (User.IsWidgetInFocusPath(SafeWidget.ToSharedRef()))
+		{
+			// The SCommonButton child is not disabled, and may be visible, wherein the Owning UCommonButtonBase is disabled. 
+			// Using Navigation::Next will first descend the widget hierarchy and focus on the SCommonButton regardless of the outer SObjectWidget state. 
+			// To avoid this, use the top level directional navigation from the SObjectWidget level
+			const TArray<EUINavigation> Directions = { EUINavigation::Right, EUINavigation::Down,  EUINavigation::Left,  EUINavigation::Up };
+			FSlateApplication::Get().NavigateToNextWidgetByDirection(User.GetUserIndex(), SafeWidget, Directions);
+		}
+	});
+}
+
 FName UWidgetLockedStateRegistration::GetStateName() const
 {
 	return StateName;
diff --git a/Engine/Plugins/Runtime/CommonUI/Source/CommonUI/Public/CommonButtonBase.h b/Engine/Plugins/Runtime/CommonUI/Source/CommonUI/Public/CommonButtonBase.h
--- a/Engine/Plugins/Runtime/CommonUI/Source/CommonUI/Public/CommonButtonBase.h
+++ b/Engine/Plugins/Runtime/CommonUI/Source/CommonUI/Public/CommonButtonBase.h
@@ -235,6 +235,8 @@ public:
 	COMMONUI_API void SetMinDesiredHeight(int32 InMinHeight);
 	COMMONUI_API void SetMinDesiredWidth(int32 InMinWidth);
 
+	COMMONUI_API inline TSharedPtr<class SCommonButton> GetCommonButton() const { return MyCommonButton; }
+
 	/** Called when the button is clicked */
 	FOnButtonDoubleClickedEvent HandleDoubleClicked;
 
@@ -310,6 +312,7 @@ public:
 	COMMONUI_API virtual void NativeDestruct() override;
 	COMMONUI_API virtual bool Initialize() override;
 	COMMONUI_API virtual void SetIsEnabled(bool bInIsEnabled) override;
+	COMMONUI_API virtual void SetVisibility(ESlateVisibility InVisibility) override;
 	COMMONUI_API virtual bool NativeIsInteractable() const override;
 	// End of UUserWidget interface
 	
@@ -804,7 +807,7 @@ protected:
 	/** True if this button should have a press and hold behavior, triggering the click when the specified hold time is met */
 	UPROPERTY(EditAnywhere, BlueprintReadOnly, Category = "Input|Hold", meta = (ExposeOnSpawn = true))
 	uint8 bRequiresHold:1;
-	
+
 	/** Press and Hold values used for Keyboard and Mouse, Gamepad and Touch, depending on the current input type */
 	UPROPERTY(EditAnywhere, BlueprintReadOnly, Category = "Input|Hold", meta = (EditCondition="bRequiresHold", ExposeOnSpawn = true))
 	TSubclassOf<UCommonUIHoldData> HoldData;
@@ -872,6 +875,10 @@ public:
 	bool bStyleNoLongerNeedsConversion;
 #endif
 
+	/** If this button is currently in focus, and is disabled, hidden, or collapsed then focus will be routed to the next available widget */
+	UPROPERTY(EditAnywhere, BlueprintReadOnly, Category = Input)
+	uint8 bNavigateToNextWidgetOnDisable : 1;
+
 protected:
 	UPROPERTY(BlueprintAssignable, Category = "Events", meta = (AllowPrivateAccess = true, DisplayName = "On Selected Changed"))
 	FCommonSelectedStateChangedBase OnSelectedChangedBase;
@@ -959,6 +966,8 @@ private:
 	/** Disables this button (called in SetIsEnabled override) */
 	void DisableButton();
 
+	void HandleImplicitFocusLost();
+
 	FText EnabledTooltipText;
 	FText DisabledTooltipText;
 
diff --git a/Engine/Source/Runtime/Slate/Private/Framework/Application/SlateApplication.cpp b/Engine/Source/Runtime/Slate/Private/Framework/Application/SlateApplication.cpp
--- a/Engine/Source/Runtime/Slate/Private/Framework/Application/SlateApplication.cpp
+++ b/Engine/Source/Runtime/Slate/Private/Framework/Application/SlateApplication.cpp
@@ -6223,6 +6223,41 @@ void FSlateApplication::NavigateToWidget(const uint32 UserIndex, const TSharedPt
 	}
 }
 
+void FSlateApplication::NavigateToNextWidgetByDirection(const uint32 InUserIndex, TSharedPtr<SWidget> InNavigationSource, const TArray<EUINavigation>& InDirections)
+{
+	FWidgetPath NavigationSourcePath;
+	if (InNavigationSource.IsValid())
+	{
+		FindPathToWidget(InNavigationSource.ToSharedRef(), NavigationSourcePath, EVisibility::All);
+	}
+	else
+	{
+		NavigationSourcePath = Users[InUserIndex]->GetFocusPath().Get();
+	}
+
+	if (!NavigationSourcePath.IsValid())
+	{
+		UE_LOG(LogSlate, Log, TEXT("Failed to navigate to next available widget. No widget currently in Focus"));
+		return;
+	}
+
+	// Traverse up the widget tree and attempt to navigate to next available widget by searching neighbor space
+	for (int32 WidgetIndex = NavigationSourcePath.Widgets.Num() - 1; WidgetIndex >= 0; --WidgetIndex)
+	{
+		const FArrangedWidget& BoundaryWidget = NavigationSourcePath.Widgets[WidgetIndex];
+
+		for (EUINavigation Direction : InDirections)
+		{
+			// Try to advance to next widget
+			const FNavigationEvent NavigationEventNext(PlatformApplication->GetModifierKeys(), InUserIndex, Direction, ENavigationGenesis::User);
+			if (AttemptNavigation(NavigationSourcePath, NavigationEventNext, FNavigationReply::Escape(), BoundaryWidget))
+			{
+				return;
+			}
+		}
+	}
+}
+
 bool FSlateApplication::AttemptNavigation(const FWidgetPath& NavigationSource, const FNavigationEvent& NavigationEvent, const FNavigationReply& NavigationReply, const FArrangedWidget& BoundaryWidget)
 {
 	if ( !NavigationSource.IsValid() )
diff --git a/Engine/Source/Runtime/Slate/Public/Framework/Application/SlateApplication.h b/Engine/Source/Runtime/Slate/Public/Framework/Application/SlateApplication.h
--- a/Engine/Source/Runtime/Slate/Public/Framework/Application/SlateApplication.h
+++ b/Engine/Source/Runtime/Slate/Public/Framework/Application/SlateApplication.h
@@ -1810,6 +1810,15 @@ public:
 	 */
 	SLATE_API void NavigateFromWidgetUnderCursor(const uint32 InUserIndex, EUINavigation InNavigationType, TSharedRef<SWindow> InWindow);
 
+	/**
+	 * Attempts to navigate to a neighboring widget by searching the given directions
+	 *
+	 * @param InUserIndex The user that is doing the navigation
+	 * @param InNavigationSource The source widget to navigate from. If null, will use current user focus
+	 * @param InDirections The direction space to search in
+	 */
+	SLATE_API void NavigateToNextWidgetByDirection(const uint32 InUserIndex, TSharedPtr<SWidget> InNavigationSource, const TArray<EUINavigation>& InDirections);
+
 	/**
 	* Given an optional widget, try and get the most suitable parent window to use with dialogs (such as file and directory pickers).
 	* This will first try and get the window that owns the widget (if provided), before falling back to using the MainFrame window.
