CommonUI: Remove ensure that was meant to be removed with the change to use FCommonActivatableSlateMetaData

#rnx

diff --git a/Engine/Plugins/Runtime/CommonUI/Source/CommonUI/Private/Widgets/CommonActivatableWidgetContainer.cpp b/Engine/Plugins/Runtime/CommonUI/Source/CommonUI/Private/Widgets/CommonActivatableWidgetContainer.cpp
--- a/Engine/Plugins/Runtime/CommonUI/Source/CommonUI/Private/Widgets/CommonActivatableWidgetContainer.cpp
+++ b/Engine/Plugins/Runtime/CommonUI/Source/CommonUI/Private/Widgets/CommonActivatableWidgetContainer.cpp
@@ -13,7 +13,7 @@
 
 UCommonActivatableWidget* ActivatableWidgetFromSlate(const TSharedPtr<SWidget>& SlateWidget)
 {
-	if (SlateWidget && SlateWidget != SNullWidget::NullWidget && ensure(SlateWidget->GetMetaData<FCommonActivatableSlateMetaData>().IsValid()))
+	if (SlateWidget && SlateWidget != SNullWidget::NullWidget && SlateWidget->GetMetaData<FCommonActivatableSlateMetaData>().IsValid())
 	{
 		UCommonActivatableWidget* ActivatableWidget = Cast<UCommonActivatableWidget>(StaticCastSharedPtr<SObjectWidget>(SlateWidget)->GetWidgetObject());
 		if (ensure(ActivatableWidget))
