[CommonButtonBase] Exposed a MaxWidth and MaxHeight variables to be able to set the max desired width/height on its SBox.
[RN] Exposed "MaxWidth" and "MaxHeight" variables to CommonButtonBase. This will set the desired max dimension on the underlying SBox.

diff --git a/Engine/Plugins/Runtime/CommonUI/Source/CommonUI/Private/CommonButtonBase.cpp b/Engine/Plugins/Runtime/CommonUI/Source/CommonUI/Private/CommonButtonBase.cpp
--- a/Engine/Plugins/Runtime/CommonUI/Source/CommonUI/Private/CommonButtonBase.cpp
+++ b/Engine/Plugins/Runtime/CommonUI/Source/CommonUI/Private/CommonButtonBase.cpp
@@ -251,6 +251,24 @@ void UCommonButtonInternalBase::SetMinDesiredWidth(int32 InMinWidth)
 	}
 }
 
+void UCommonButtonInternalBase::SetMaxDesiredHeight(int32 InMaxHeight)
+{
+	MaxHeight = InMaxHeight;
+	if (MyBox.IsValid())
+	{
+		MyBox->SetMaxDesiredHeight(InMaxHeight > 0 ? InMaxHeight : FOptionalSize());
+	}
+}
+
+void UCommonButtonInternalBase::SetMaxDesiredWidth(int32 InMaxWidth)
+{
+	MaxWidth = InMaxWidth;
+	if (MyBox.IsValid())
+	{
+		MyBox->SetMaxDesiredWidth(InMaxWidth > 0 ? InMaxWidth : FOptionalSize());
+	}
+}
+
 TSharedRef<SWidget> UCommonButtonInternalBase::RebuildWidget()
 {
 	MyButton = MyCommonButton = SNew(SCommonButton)
@@ -270,6 +288,8 @@ TSharedRef<SWidget> UCommonButtonInternalBase::RebuildWidget()
 	MyBox = SNew(SBox)
 		.MinDesiredWidth(MinWidth)
 		.MinDesiredHeight(MinHeight)
+		.MaxDesiredWidth(MaxWidth > 0 ? MaxWidth : FOptionalSize())
+		.MaxDesiredHeight(MaxHeight > 0 ? MaxHeight : FOptionalSize())
 		.HAlign(HAlign_Fill)
 		.VAlign(VAlign_Fill)
 		[
@@ -341,6 +361,8 @@ UCommonButtonBase::UCommonButtonBase(const FObjectInitializer& ObjectInitializer
 	: Super(ObjectInitializer)
 	, MinWidth(0)
 	, MinHeight(0)
+	, MaxWidth(0)
+	, MaxHeight(0)
 	, bApplyAlphaOnDisable(true)
 	, bLocked(false)
 	, bSelectable(false)
@@ -1052,6 +1074,32 @@ void UCommonButtonBase::RefreshDimensions()
 		const UCommonButtonStyle* const StyleCDO = GetStyleCDO();
 		RootButton->SetMinDesiredWidth(FMath::Max(MinWidth, StyleCDO ? StyleCDO->MinWidth : 0));
 		RootButton->SetMinDesiredHeight(FMath::Max(MinHeight, StyleCDO ? StyleCDO->MinHeight : 0));
+		
+		if (!StyleCDO)
+		{
+			RootButton->SetMaxDesiredWidth(MaxWidth);
+			RootButton->SetMaxDesiredHeight(MaxHeight);
+		}
+		else
+		{
+			if (MaxWidth > 0 && StyleCDO->MaxWidth > 0)
+			{
+				RootButton->SetMaxDesiredWidth(FMath::Min(MaxWidth, StyleCDO->MaxWidth));
+			}
+			else
+			{
+				RootButton->SetMaxDesiredWidth(FMath::Max(MaxWidth, StyleCDO->MaxWidth));
+			}
+
+			if (MaxHeight > 0 && StyleCDO->MaxHeight > 0)
+			{
+				RootButton->SetMaxDesiredHeight(FMath::Min(MaxHeight, StyleCDO->MaxHeight));
+			}
+			else
+			{
+				RootButton->SetMaxDesiredHeight(FMath::Max(MaxHeight, StyleCDO->MaxHeight));
+			}
+		}
 	}
 }
 
@@ -1211,6 +1259,14 @@ void UCommonButtonBase::SetMinDimensions(int32 InMinWidth, int32 InMinHeight)
 	RefreshDimensions();
 }
 
+void UCommonButtonBase::SetMaxDimensions(int32 InMaxWidth, int32 InMaxHeight)
+{
+	MaxWidth = InMaxWidth;
+	MaxHeight = InMaxHeight;
+	
+	RefreshDimensions();
+}
+
 void UCommonButtonBase::SetTriggeredInputAction(const FDataTableRowHandle &InputActionRow)
 {
 	UnbindTriggeringInputActionToClick();
diff --git a/Engine/Plugins/Runtime/CommonUI/Source/CommonUI/Public/CommonButtonBase.h b/Engine/Plugins/Runtime/CommonUI/Source/CommonUI/Public/CommonButtonBase.h
--- a/Engine/Plugins/Runtime/CommonUI/Source/CommonUI/Public/CommonButtonBase.h
+++ b/Engine/Plugins/Runtime/CommonUI/Source/CommonUI/Public/CommonButtonBase.h
@@ -124,6 +124,14 @@ public:
 	/** The minimum height of buttons using this style */
 	UPROPERTY(EditDefaultsOnly, BlueprintReadOnly, Category = "Properties")
 	int32 MinHeight;
+	
+	/** The maximum width of buttons using this style */
+	UPROPERTY(EditDefaultsOnly, BlueprintReadOnly, Category = "Properties")
+	int32 MaxWidth;
+	
+	/** The maximum height of buttons using this style */
+    UPROPERTY(EditDefaultsOnly, BlueprintReadOnly, Category = "Properties")
+    int32 MaxHeight;
 
 	/** The text style to use when un-selected */
 	UPROPERTY(EditDefaultsOnly, BlueprintReadOnly, Category = "Properties")
@@ -234,6 +242,8 @@ public:
 
 	COMMONUI_API void SetMinDesiredHeight(int32 InMinHeight);
 	COMMONUI_API void SetMinDesiredWidth(int32 InMinWidth);
+	COMMONUI_API void SetMaxDesiredHeight(int32 InMaxHeight);
+	COMMONUI_API void SetMaxDesiredWidth(int32 InMaxWidth);
 
 	COMMONUI_API inline TSharedPtr<class SCommonButton> GetCommonButton() const { return MyCommonButton; }
 
@@ -274,6 +284,14 @@ protected:
 	/** The minimum height of the button */
 	UPROPERTY()
 	int32 MinHeight;
+	
+	/** The maximum width of the button */
+	UPROPERTY()
+	int32 MaxWidth;
+	
+	/** The maximum height of the button */
+	UPROPERTY()
+	int32 MaxHeight;
 
 	/** If true, this button is enabled. */
 	UPROPERTY()
@@ -423,6 +441,10 @@ public:
 	UFUNCTION(BlueprintCallable, Category = "Common Button|Setters")
 	COMMONUI_API void SetMinDimensions(int32 InMinWidth, int32 InMinHeight);
 
+	/** Sets the maximum dimensions of this button */
+	UFUNCTION(BlueprintCallable, Category = "Common Button|Setters")
+	COMMONUI_API void SetMaxDimensions(int32 InMaxWidth, int32 InMaxHeight);
+
 	/** Updates the current triggered action */
 	UFUNCTION(BlueprintCallable, Category = "Common Button|Setters")
 	COMMONUI_API void SetTriggeredInputAction(const FDataTableRowHandle &InputActionRow);
@@ -722,6 +744,14 @@ protected:
 	/** The minimum height of the button (only used if greater than the style's minimum) */
 	UPROPERTY(EditAnywhere, BlueprintReadOnly, Category = Layout, meta = (ClampMin = "0"))
 	int32 MinHeight;
+	
+	/** The maximum width of the button (only used if greater than the style's maximum) */
+	UPROPERTY(EditAnywhere, BlueprintReadOnly, Category = Layout, meta = (ClampMin = "0"))
+	int32 MaxWidth;
+	
+	/** The maximum height of the button (only used if greater than the style's maximum) */
+	UPROPERTY(EditAnywhere, BlueprintReadOnly, Category = Layout, meta = (ClampMin = "0"))
+	int32 MaxHeight;
 
 	/** References the button style asset that defines a style in multiple sizes */
 	UPROPERTY(EditAnywhere, BlueprintReadOnly, Category = Style, meta = (ExposeOnSpawn = true))
