PR #12879: [CommonUI] Add SetStyle function to CommonRichTextBlock

[RN]

diff --git a/Engine/Plugins/Runtime/CommonUI/Source/CommonUI/Private/CommonRichTextBlock.cpp b/Engine/Plugins/Runtime/CommonUI/Source/CommonUI/Private/CommonRichTextBlock.cpp
--- a/Engine/Plugins/Runtime/CommonUI/Source/CommonUI/Private/CommonRichTextBlock.cpp
+++ b/Engine/Plugins/Runtime/CommonUI/Source/CommonUI/Private/CommonRichTextBlock.cpp
@@ -412,6 +412,12 @@ void UCommonRichTextBlock::SetText(const FText& InText)
 	}
 }
 
+void UCommonRichTextBlock::SetStyle(const TSubclassOf<UCommonTextStyle>& InStyle)
+{
+	DefaultTextStyleOverrideClass = InStyle;
+	SynchronizeProperties();
+}
+
 void UCommonRichTextBlock::SetScrollingEnabled(bool bInIsScrollingEnabled)
 {
 	bIsScrollingEnabled = bInIsScrollingEnabled;
diff --git a/Engine/Plugins/Runtime/CommonUI/Source/CommonUI/Public/CommonRichTextBlock.h b/Engine/Plugins/Runtime/CommonUI/Source/CommonUI/Public/CommonRichTextBlock.h
--- a/Engine/Plugins/Runtime/CommonUI/Source/CommonUI/Public/CommonRichTextBlock.h
+++ b/Engine/Plugins/Runtime/CommonUI/Source/CommonUI/Public/CommonRichTextBlock.h
@@ -43,6 +43,9 @@ public:
 	virtual void ReleaseSlateResources(bool bReleaseChildren) override;
 	virtual void SetText(const FText& InText) override;
 
+	UFUNCTION(BlueprintCallable, Category = "Common Rich Text")
+	void SetStyle(const TSubclassOf<UCommonTextStyle>& InStyle);
+
 	UFUNCTION(BlueprintCallable, Category = "Common Rich Text|Scroll Style")
 	void SetScrollingEnabled(bool bInIsScrollingEnabled);
 
