Update CommonActionWidget to mark SetInputAction functions as virtual

diff --git a/Engine/Plugins/Runtime/CommonUI/Source/CommonUI/Public/CommonActionWidget.h b/Engine/Plugins/Runtime/CommonUI/Source/CommonUI/Public/CommonActionWidget.h
--- a/Engine/Plugins/Runtime/CommonUI/Source/CommonUI/Public/CommonActionWidget.h
+++ b/Engine/Plugins/Runtime/CommonUI/Source/CommonUI/Public/CommonActionWidget.h
@@ -46,19 +46,19 @@ public:
 	UMaterialInstanceDynamic* GetIconDynamicMaterial();
 
 	UFUNCTION(BlueprintCallable, Category = CommonActionWidget)
-	void SetEnhancedInputAction(UInputAction* InInputAction);
+	virtual void SetEnhancedInputAction(UInputAction* InInputAction);
 
 	UFUNCTION(BlueprintCallable, Category = CommonActionWidget)
 	const UInputAction* GetEnhancedInputAction() const;
 
 	UFUNCTION(BlueprintCallable, Category = CommonActionWidget)
-	void SetInputAction(FDataTableRowHandle InputActionRow);
+	virtual void SetInputAction(FDataTableRowHandle InputActionRow);
 
 	UFUNCTION(BlueprintCallable, Category = CommonActionWidget)
-	void SetInputActionBinding(FUIActionBindingHandle BindingHandle);
+	virtual void SetInputActionBinding(FUIActionBindingHandle BindingHandle);
 
 	UFUNCTION(BlueprintCallable, Category = CommonActionWidget)
-	void SetInputActions(TArray<FDataTableRowHandle> NewInputActions);
+	virtual void SetInputActions(TArray<FDataTableRowHandle> NewInputActions);
 
 	UFUNCTION(BlueprintCallable, Category = CommonActionWidget)
 	void SetIconRimBrush(FSlateBrush InIconRimBrush);
