[Common Input/UI] Add some auto fix up options for adding new platform names to common input data table assets

Also ran a PF for a full cook: 67e32b889e3d98359761e810

diff --git a/Engine/Plugins/Runtime/CommonUI/Source/CommonInput/Public/CommonInputSettings.h b/Engine/Plugins/Runtime/CommonUI/Source/CommonInput/Public/CommonInputSettings.h
--- a/Engine/Plugins/Runtime/CommonUI/Source/CommonInput/Public/CommonInputSettings.h
+++ b/Engine/Plugins/Runtime/CommonUI/Source/CommonInput/Public/CommonInputSettings.h
@@ -64,6 +64,8 @@ public:
 
 	TObjectPtr<UCommonInputActionDomainTable> GetActionDomainTable() const { return ActionDomainTablePtr; }
 
+	const TMap<FName, FName>& GetPlatformNameUpgradeMap() const { return PlatformNameUpgrades; }
+
 public:
 
 	/** Static version of enhanced input support check, exists to hide based on edit condition */
@@ -119,6 +121,14 @@ private:
 	/** Create a derived asset from UCommonInputActionDomainTable to store ordered ActionDomain data for your game */
 	UPROPERTY(config, EditAnywhere, Category = "Action Domain")
 	TSoftObjectPtr<UCommonInputActionDomainTable> ActionDomainTable;
+	
+	/**
+	* A map of Common Input platform names to a new one, which you can use
+	* to upgrade your Input Action data tables if you add a new platform to your
+	* project and wish to copy from some existing data
+	*/
+	UPROPERTY(config, EditAnywhere, Category = "Input", meta = (ConfigRestartRequired = true))
+	TMap<FName, FName> PlatformNameUpgrades;
 
 private:
 	void LoadInputData();
diff --git a/Engine/Plugins/Runtime/CommonUI/Source/CommonUI/Private/CommonUITypes.cpp b/Engine/Plugins/Runtime/CommonUI/Source/CommonUI/Private/CommonUITypes.cpp
--- a/Engine/Plugins/Runtime/CommonUI/Source/CommonUI/Private/CommonUITypes.cpp
+++ b/Engine/Plugins/Runtime/CommonUI/Source/CommonUI/Private/CommonUITypes.cpp
@@ -5,6 +5,7 @@
 #include "CommonInputBaseTypes.h"
 #include "CommonInputSubsystem.h"
 #include "CommonUIPrivate.h"
+#include "CommonInputSettings.h"
 #include "Engine/Blueprint.h"
 #include "EnhancedActionKeyMapping.h"
 #include "EnhancedInputSubsystems.h"
@@ -13,9 +14,11 @@
 #include "ICommonInputModule.h"
 #include "InputAction.h"
 #include "InputActionValue.h"
+#include "Misc/ConfigCacheIni.h"
 #include "PlayerMappableKeySettings.h"
 #include "Styling/SlateTypes.h"
 #include "Styling/StyleDefaults.h"
+#include "UObject/FortniteReleaseBranchCustomObjectVersion.h"
 
 #include UE_INLINE_GENERATED_CPP_BY_NAME(CommonUITypes)
 
@@ -147,6 +150,47 @@ FSlateBrush FCommonInputActionDataBase::GetCurrentInputActionIcon(const UCommonI
 	return *FStyleDefaults::GetNoBrush();
 }
 
+bool FCommonInputActionDataBase::Serialize(FArchive& Ar)
+{
+	Ar.UsingCustomVersion(FFortniteReleaseBranchCustomObjectVersion::GUID);
+
+	// doesn't actually serialize, just write the custom version for PostSerialize
+	return false;
+}
+
+void FCommonInputActionDataBase::PostSerialize(const FArchive& Ar)
+{
+	// During load, if the asset version is lower then when we added our platform upgrade path
+	// then we should run that upgrade path data
+	if (Ar.IsLoading())
+	{
+		if (Ar.CustomVer(FFortniteReleaseBranchCustomObjectVersion::GUID) < FFortniteReleaseBranchCustomObjectVersion::CommonUIPlatformNamingUpgradeOption)
+		{
+			// Upgrade the gamepad override path stuff!
+			const TMap<FName, FName>& PlatformNameRedirects = GetDefault<UCommonInputSettings>()->GetPlatformNameUpgradeMap();
+
+			for (const TPair<FName, FName>& Redirect : PlatformNameRedirects)
+			{
+				// If there is a platform override defined for one of the platforms we want to copy from...
+				const FCommonInputTypeInfo* FromPadData = GamepadInputOverrides.Find(Redirect.Key);
+				if (!FromPadData)
+				{
+					continue;
+				}
+
+				// Make sure that the override is not already defined for this platform
+				if (GamepadInputOverrides.Contains(Redirect.Value))
+				{
+					continue;
+				}
+
+				// Add the new gamepad override data and initialize it to the data which we had for the older platform name
+				GamepadInputOverrides.Add(Redirect.Value, FCommonInputTypeInfo (*FromPadData));
+			}
+		}
+	}
+}
+
 void FCommonInputActionDataBase::OnPostDataImport(const UDataTable* InDataTable, const FName InRowName, TArray<FString>& OutCollectedImportProblems)
 {
 #if WITH_EDITOR
diff --git a/Engine/Plugins/Runtime/CommonUI/Source/CommonUI/Public/CommonUITypes.h b/Engine/Plugins/Runtime/CommonUI/Source/CommonUI/Public/CommonUITypes.h
--- a/Engine/Plugins/Runtime/CommonUI/Source/CommonUI/Public/CommonUITypes.h
+++ b/Engine/Plugins/Runtime/CommonUI/Source/CommonUI/Public/CommonUITypes.h
@@ -178,6 +178,19 @@ public:
 		check(this);
 		return !(*this == Other);
 	}
+
+	bool Serialize(FArchive& Ar);
+	void PostSerialize(const FArchive& Ar);
+};
+
+template<>
+struct TStructOpsTypeTraits<FCommonInputActionDataBase> : public TStructOpsTypeTraitsBase2<FCommonInputActionDataBase>
+{
+	enum
+	{
+		WithSerializer = true,
+		WithPostSerialize = true,
+	};
 };
 
 /** 
diff --git a/Engine/Source/Runtime/Core/Public/UObject/FortniteReleaseBranchCustomObjectVersions.inl b/Engine/Source/Runtime/Core/Public/UObject/FortniteReleaseBranchCustomObjectVersions.inl
--- a/Engine/Source/Runtime/Core/Public/UObject/FortniteReleaseBranchCustomObjectVersions.inl
+++ b/Engine/Source/Runtime/Core/Public/UObject/FortniteReleaseBranchCustomObjectVersions.inl
@@ -57,5 +57,8 @@ FORTNITE_RELEASE_VERSION(DynamicMeshAttributesSerializeBones, 15)
 // Add option for sanitizing output attribute names for all PCG data getters
 FORTNITE_RELEASE_VERSION(OptionSanitizeOutputAttributeNamesPCG, 16)
 
+// Add automatic platform naming fix up for CommonUI input action data tables
+FORTNITE_RELEASE_VERSION(CommonUIPlatformNamingUpgradeOption, 17)
+
 // -----<new versions can be added above this line>-------------------------------------------------
 #undef FORTNITE_RELEASE_VERSION
\ No newline at end of file
