Exposed a virtual GetTimespanStartingTime method for UCommonDateTimeTextBlock which can be overridden to change what "now" means in derived classes

diff --git a/Engine/Plugins/Runtime/CommonUI/Source/CommonUI/Private/CommonDateTimeTextBlock.cpp b/Engine/Plugins/Runtime/CommonUI/Source/CommonUI/Private/CommonDateTimeTextBlock.cpp
--- a/Engine/Plugins/Runtime/CommonUI/Source/CommonUI/Private/CommonDateTimeTextBlock.cpp
+++ b/Engine/Plugins/Runtime/CommonUI/Source/CommonUI/Private/CommonDateTimeTextBlock.cpp
@@ -62,7 +62,7 @@ void UCommonDateTimeTextBlock::SetDateTimeValue(const FDateTime InDateTime, bool
 
 void UCommonDateTimeTextBlock::SetTimespanValue(const FTimespan InTimespan)
 {
-	SetDateTimeValue(FDateTime::Now() + InTimespan, true);
+	SetDateTimeValue(GetTimespanStartingTime() + InTimespan, true);
 }
 
 void UCommonDateTimeTextBlock::SetCountDownCompletionText(const FText InCompletionText)
@@ -84,7 +84,7 @@ void UCommonDateTimeTextBlock::UpdateUnderlyingText()
 	TOptional<FText> TextToSet = TOptional<FText>();
 	if (bShowAsCountdown)
 	{
-		FTimespan Remaining = DateTime - FDateTime::Now();
+		FTimespan Remaining = DateTime - GetTimespanStartingTime();
 		if (TimerTickHandle.IsValid() && ShouldClearTimer(Remaining))
 		{
 			UWorld* const World = GetWorld();
@@ -122,6 +122,11 @@ void UCommonDateTimeTextBlock::UpdateUnderlyingText()
 	}
 }
 
+FDateTime UCommonDateTimeTextBlock::GetTimespanStartingTime() const
+{
+	return FDateTime::Now();
+}
+
 bool UCommonDateTimeTextBlock::ShouldClearTimer(const FTimespan& TimeRemaining) const
 {
 	return TimeRemaining.GetTotalSeconds() < 1.0;
diff --git a/Engine/Plugins/Runtime/CommonUI/Source/CommonUI/Public/CommonDateTimeTextBlock.h b/Engine/Plugins/Runtime/CommonUI/Source/CommonUI/Public/CommonDateTimeTextBlock.h
--- a/Engine/Plugins/Runtime/CommonUI/Source/CommonUI/Public/CommonDateTimeTextBlock.h
+++ b/Engine/Plugins/Runtime/CommonUI/Source/CommonUI/Public/CommonDateTimeTextBlock.h
@@ -38,6 +38,7 @@ protected:
 	
 	void UpdateUnderlyingText();
 
+	virtual FDateTime GetTimespanStartingTime() const;
 	virtual bool ShouldClearTimer(const FTimespan& TimeRemaining) const;
 
 	virtual TOptional<FText> FormatTimespan(const FTimespan& InTimespan) const;
