Common UI:
- Add bound actions updated broadcast when deactivating an action domain

#rnx

diff --git a/Engine/Plugins/Runtime/CommonUI/Source/CommonUI/Private/Input/CommonUIActionRouterBase.cpp b/Engine/Plugins/Runtime/CommonUI/Source/CommonUI/Private/Input/CommonUIActionRouterBase.cpp
--- a/Engine/Plugins/Runtime/CommonUI/Source/CommonUI/Private/Input/CommonUIActionRouterBase.cpp
+++ b/Engine/Plugins/Runtime/CommonUI/Source/CommonUI/Private/Input/CommonUIActionRouterBase.cpp
@@ -934,6 +934,7 @@ void UCommonUIActionRouterBase::HandleRootNodeDeactivated(TWeakPtr<FActivatableT
 		}
 
 		RefreshActionDomainLeafNodeConfig();
+		OnBoundActionsUpdated().Broadcast();
 	}
 }
 
