CommonUI: Workaround SlateGI bug that causes a flicker (and other issues) with CommonAnimatedSwitcher

#rnx

diff --git a/Engine/Plugins/Runtime/CommonUI/Source/CommonUI/Private/Slate/SCommonAnimatedSwitcher.cpp b/Engine/Plugins/Runtime/CommonUI/Source/CommonUI/Private/Slate/SCommonAnimatedSwitcher.cpp
--- a/Engine/Plugins/Runtime/CommonUI/Source/CommonUI/Private/Slate/SCommonAnimatedSwitcher.cpp
+++ b/Engine/Plugins/Runtime/CommonUI/Source/CommonUI/Private/Slate/SCommonAnimatedSwitcher.cpp
@@ -288,6 +288,14 @@ EActiveTimerReturnType SCommonAnimatedSwitcher::UpdateTransition(double InCurren
 		OnIsTransitioningChanged.ExecuteIfBound(false);
 		bIsTransitionTimerRegistered = false;
 
+		// TODO - UE-256488 - There's a bug with SlateGI where setting the visibility of this widget in the same frame as becoming non-volatile (by stopping our timer) will cause this widget to disappear for that frame, so we stay volatile for 1 more frame
+		ForceVolatile(true);
+		FTSTicker::GetCoreTicker().AddTicker(FTickerDelegate::CreateSPLambda(this, [this](float DeltaTime)
+		{
+			ForceVolatile(false);
+			return false;
+		}));
+
 		return EActiveTimerReturnType::Stop;
 	}
 	else
