PR #13088: CommonUI - TabListWidgetBase - Added a bool to control whether the tab navigation should wrap or not.

#rnx

diff --git a/Engine/Plugins/Runtime/CommonUI/Source/CommonUI/Private/CommonTabListWidgetBase.cpp b/Engine/Plugins/Runtime/CommonUI/Source/CommonUI/Private/CommonTabListWidgetBase.cpp
--- a/Engine/Plugins/Runtime/CommonUI/Source/CommonUI/Private/CommonTabListWidgetBase.cpp
+++ b/Engine/Plugins/Runtime/CommonUI/Source/CommonUI/Private/CommonTabListWidgetBase.cpp
@@ -511,7 +511,7 @@ void UCommonTabListWidgetBase::HandleNextTabAction()
 {
 	if (ensure(TabButtonGroup))
 	{
-		TabButtonGroup->SelectNextButton();
+		TabButtonGroup->SelectNextButton(bShouldWrapNavigation);
 	}
 }
 
@@ -524,7 +524,7 @@ void UCommonTabListWidgetBase::HandlePreviousTabAction()
 {
 	if (ensure(TabButtonGroup))
 	{
-		TabButtonGroup->SelectPreviousButton();
+		TabButtonGroup->SelectPreviousButton(bShouldWrapNavigation);
 	}
 }
 
diff --git a/Engine/Plugins/Runtime/CommonUI/Source/CommonUI/Public/CommonTabListWidgetBase.h b/Engine/Plugins/Runtime/CommonUI/Source/CommonUI/Public/CommonTabListWidgetBase.h
--- a/Engine/Plugins/Runtime/CommonUI/Source/CommonUI/Public/CommonTabListWidgetBase.h
+++ b/Engine/Plugins/Runtime/CommonUI/Source/CommonUI/Public/CommonTabListWidgetBase.h
@@ -209,6 +209,10 @@ protected:
 	UPROPERTY(EditAnywhere, BlueprintReadOnly, Category = TabList, meta = (ExposeOnSpawn = "true"))
 	bool bAutoListenForInput;
 
+	/** Whether pressing next/prev tab on the last/first tab should wrap selection to the beginning/end or stay at the end/beginning. */
+	UPROPERTY(EditAnywhere, BlueprintReadOnly, Category = TabList, meta = (ExposeOnSpawn = "true"))
+	bool bShouldWrapNavigation = true;
+
 	/**
 	* Whether to defer until next tick rebuilding tab list when inserting new tab (rather than adding to the end).
 	* Useful if inserting multiple tabs in the same tick as the tab list will only be rebuilt once.
