[CommonUI] Restoring focus to Action Domain Root Nodes

- Added logic to try restore focus to Action Domain Root Nodes in case there is no Active Root.

#tests In Editor
#rnx

diff --git a/Engine/Plugins/Runtime/CommonUI/Source/CommonUI/Private/Input/CommonUIActionRouterBase.cpp b/Engine/Plugins/Runtime/CommonUI/Source/CommonUI/Private/Input/CommonUIActionRouterBase.cpp
--- a/Engine/Plugins/Runtime/CommonUI/Source/CommonUI/Private/Input/CommonUIActionRouterBase.cpp
+++ b/Engine/Plugins/Runtime/CommonUI/Source/CommonUI/Private/Input/CommonUIActionRouterBase.cpp
@@ -685,7 +685,8 @@ void UCommonUIActionRouterBase::FlushInput()
 
 bool UCommonUIActionRouterBase::IsWidgetInActiveRoot(const UCommonActivatableWidget* Widget) const
 {
-	if (Widget && ActiveRootNode)
+	FActivatableTreeRootPtr RootNode = ActiveRootNode ? ActiveRootNode : FindActiveActionDomainRootNode();
+	if (Widget && RootNode)
 	{
 		TSharedPtr<SWidget> WidgetWalker = Widget->GetCachedWidget();
 		while (WidgetWalker)
@@ -694,7 +695,7 @@ bool UCommonUIActionRouterBase::IsWidgetInActiveRoot(const UCommonActivatableWid
 			{
 				if (UCommonActivatableWidget* CandidateActivatable = Cast<UCommonActivatableWidget>(StaticCastSharedPtr<SObjectWidget>(WidgetWalker)->GetWidgetObject()))
 				{
-					if (CandidateActivatable == ActiveRootNode->GetWidget())
+					if (CandidateActivatable == RootNode->GetWidget())
 					{
 						return true;
 					}
@@ -703,6 +704,7 @@ bool UCommonUIActionRouterBase::IsWidgetInActiveRoot(const UCommonActivatableWid
 			WidgetWalker = WidgetWalker->GetParentWidget();
 		}
 	}
+
 	return false;
 }
 
@@ -1442,17 +1444,19 @@ void UCommonUIActionRouterBase::RegisterWidgetBindings(const FActivatableTreeNod
 
 void UCommonUIActionRouterBase::RefreshActiveRootFocusRestorationTarget() const
 {
-	if (ActiveRootNode)
+	FActivatableTreeRootPtr RootNode = ActiveRootNode ? ActiveRootNode : FindActiveActionDomainRootNode();
+	if (RootNode)
 	{
-		ActiveRootNode->RefreshCachedRestorationTarget();
+		RootNode->RefreshCachedRestorationTarget();
 	}
 }
 
 void UCommonUIActionRouterBase::RefreshActiveRootFocus()
 {
-	if (ActiveRootNode)
+	FActivatableTreeRootPtr RootNode = ActiveRootNode ? ActiveRootNode : FindActiveActionDomainRootNode();
+	if (RootNode)
 	{
-		ActiveRootNode->FocusLeafmostNode();
+		RootNode->FocusLeafmostNode();
 	}
 }
 
