PR #13096: CommonUI - ActivatableWidgetContainer - Added a bool to reset pool of widgets when releasing slate resources (usually when loading a new map).

#rnx

diff --git a/Engine/Plugins/Runtime/CommonUI/Source/CommonUI/Private/Widgets/CommonActivatableWidgetContainer.cpp b/Engine/Plugins/Runtime/CommonUI/Source/CommonUI/Private/Widgets/CommonActivatableWidgetContainer.cpp
--- a/Engine/Plugins/Runtime/CommonUI/Source/CommonUI/Private/Widgets/CommonActivatableWidgetContainer.cpp
+++ b/Engine/Plugins/Runtime/CommonUI/Source/CommonUI/Private/Widgets/CommonActivatableWidgetContainer.cpp
@@ -130,7 +130,14 @@ void UCommonActivatableWidgetContainerBase::ReleaseSlateResources(bool bReleaseC
 	ReleasedWidgets.Empty();
 	WidgetList.Reset();
 
-	GeneratedWidgetsPool.ReleaseAll(true);
+	if (bResetPoolWhenReleasingSlateResources)
+	{
+		GeneratedWidgetsPool.ResetPool();
+	}
+	else
+	{
+		GeneratedWidgetsPool.ReleaseAll(true);
+	}
 }
 
 void UCommonActivatableWidgetContainerBase::OnWidgetRebuilt()
diff --git a/Engine/Plugins/Runtime/CommonUI/Source/CommonUI/Public/Widgets/CommonActivatableWidgetContainer.h b/Engine/Plugins/Runtime/CommonUI/Source/CommonUI/Public/Widgets/CommonActivatableWidgetContainer.h
--- a/Engine/Plugins/Runtime/CommonUI/Source/CommonUI/Public/Widgets/CommonActivatableWidgetContainer.h
+++ b/Engine/Plugins/Runtime/CommonUI/Source/CommonUI/Public/Widgets/CommonActivatableWidgetContainer.h
@@ -113,6 +113,14 @@ protected:
 	UPROPERTY(EditAnywhere, Category = Transition)
 	float TransitionDuration = 0.4f;
 
+	/**
+	 * Whether to completely reset the pool of widgets when slate resources are released.
+	 * This usually happens when changing maps. You may not want to have all frontend screens loaded taking up memory while in game and vice versa.
+	 * Enabling this means widgets will have to be loaded again when re-entering the map next time.
+	 */
+	UPROPERTY(EditAnywhere, Category = Performance)
+	bool bResetPoolWhenReleasingSlateResources = false;
+
 	/**
 	 * Controls how we will choose another widget if a transitioning widget is removed during the transition.
 	 * Note for Queues and Stacks, ECommonSwitcherTransitionFallbackStrategy::Previous is a good option.
