CommonUI: Update CommonButtonBase to consistently focus the child SCommonButton

#rnx

diff --git a/Engine/Plugins/Runtime/CommonUI/Source/CommonUI/Private/CommonButtonBase.cpp b/Engine/Plugins/Runtime/CommonUI/Source/CommonUI/Private/CommonButtonBase.cpp
--- a/Engine/Plugins/Runtime/CommonUI/Source/CommonUI/Private/CommonButtonBase.cpp
+++ b/Engine/Plugins/Runtime/CommonUI/Source/CommonUI/Private/CommonButtonBase.cpp
@@ -1531,18 +1531,14 @@ FReply UCommonButtonBase::NativeOnFocusReceived(const FGeometry& InGeometry, con
 {
 	FReply Reply = Super::NativeOnFocusReceived(InGeometry, InFocusEvent);
 
-	HandleFocusReceived();
+	if (!Reply.IsEventHandled() && RootButton.IsValid() && RootButton->GetCommonButton().IsValid())
+	{
+		Reply = FReply::Handled().SetUserFocus(RootButton->GetCommonButton().ToSharedRef(), InFocusEvent.GetCause());
+	}
 
 	return Reply;
 }
 
-void UCommonButtonBase::NativeOnFocusLost(const FFocusEvent& InFocusEvent)
-{
-	Super::NativeOnFocusLost(InFocusEvent);
-
-	HandleFocusLost();
-}
-
 void UCommonButtonBase::NativeOnSelected(bool bBroadcast)
 {
 	BP_OnSelected();
diff --git a/Engine/Plugins/Runtime/CommonUI/Source/CommonUI/Public/CommonButtonBase.h b/Engine/Plugins/Runtime/CommonUI/Source/CommonUI/Public/CommonButtonBase.h
--- a/Engine/Plugins/Runtime/CommonUI/Source/CommonUI/Public/CommonButtonBase.h
+++ b/Engine/Plugins/Runtime/CommonUI/Source/CommonUI/Public/CommonButtonBase.h
@@ -538,7 +538,6 @@ protected:
 	COMMONUI_API virtual void PostLoad() override;
 	COMMONUI_API virtual void SynchronizeProperties() override;
 	COMMONUI_API virtual FReply NativeOnFocusReceived(const FGeometry& InGeometry, const FFocusEvent& InFocusEvent) override;
-	COMMONUI_API virtual void NativeOnFocusLost(const FFocusEvent& InFocusEvent) override;
 
 #if WITH_EDITOR
 	COMMONUI_API virtual void OnCreationFromPalette() override;
