Check RootButton Validity before attempting to dereference it. RootButton is checked everywhere else already, which makes sense as it is a Weak Object Ptr

Crash fix for https://jira.it.epicgames.com/browse/PLAY-81781

diff --git a/Engine/Plugins/Runtime/CommonUI/Source/CommonUI/Private/CommonButtonBase.cpp b/Engine/Plugins/Runtime/CommonUI/Source/CommonUI/Private/CommonButtonBase.cpp
--- a/Engine/Plugins/Runtime/CommonUI/Source/CommonUI/Private/CommonButtonBase.cpp
+++ b/Engine/Plugins/Runtime/CommonUI/Source/CommonUI/Private/CommonButtonBase.cpp
@@ -745,30 +745,33 @@ void UCommonButtonBase::SetIsInteractionEnabled(bool bInIsInteractionEnabled)
 
 	bInteractionEnabled = bInIsInteractionEnabled;
 
-	if (bInteractionEnabled)
+	if (RootButton.IsValid())
 	{
-		// If this is a selected and not-toggleable button, don't enable root button interaction
-		if (!GetSelected() || bToggleable)
+		if (bInteractionEnabled)
 		{
-			RootButton->SetInteractionEnabled(true);
-		}
+			// If this is a selected and not-toggleable button, don't enable root button interaction
+			if (!GetSelected() || bToggleable)
+			{
+				RootButton->SetInteractionEnabled(true);
+			}
 
-		if (bApplyAlphaOnDisable)
-		{
-			FLinearColor ButtonColor = RootButton->GetColorAndOpacity();
-			ButtonColor.A = 1.f;
-			RootButton->SetColorAndOpacity(ButtonColor);
+			if (bApplyAlphaOnDisable)
+			{
+				FLinearColor ButtonColor = RootButton->GetColorAndOpacity();
+				ButtonColor.A = 1.f;
+				RootButton->SetColorAndOpacity(ButtonColor);
+			}
 		}
-	}
-	else
-	{
-		RootButton->SetInteractionEnabled(false);
-
-		if (bApplyAlphaOnDisable)
+		else
 		{
-			FLinearColor ButtonColor = RootButton->GetColorAndOpacity();
-			ButtonColor.A = 0.5f;
-			RootButton->SetColorAndOpacity(ButtonColor);
+			RootButton->SetInteractionEnabled(false);
+
+			if (bApplyAlphaOnDisable)
+			{
+				FLinearColor ButtonColor = RootButton->GetColorAndOpacity();
+				ButtonColor.A = 0.5f;
+				RootButton->SetColorAndOpacity(ButtonColor);
+			}
 		}
 	}
 
@@ -957,13 +960,16 @@ void UCommonButtonBase::SetIsToggleable(bool bInIsToggleable)
 	bToggleable = bInIsToggleable;
 
 	// Update interactability.
-	if (!GetSelected() || bToggleable)
-	{
-		RootButton->SetInteractionEnabled(bInteractionEnabled);
-	}
-	else if (GetSelected() && !bToggleable)
+	if (RootButton.IsValid())
 	{
-		RootButton->SetInteractionEnabled(bInteractableWhenSelected);
+		if (!GetSelected() || bToggleable)
+		{
+			RootButton->SetInteractionEnabled(bInteractionEnabled);
+		}
+		else if (GetSelected() && !bToggleable)
+		{
+			RootButton->SetInteractionEnabled(bInteractableWhenSelected);
+		}
 	}
 
 	UpdateInputActionWidgetVisibility();
@@ -1042,7 +1048,10 @@ void UCommonButtonBase::SetSelectedInternal(bool bInSelected, bool bAllowSound /
 		{
 			// If the button isn't toggleable, then disable interaction with the root button while selected
 			// The prevents us getting unnecessary click noises and events
-			RootButton->SetInteractionEnabled(bInteractableWhenSelected);
+			if (RootButton.IsValid())
+			{
+				RootButton->SetInteractionEnabled(bInteractableWhenSelected);
+			}
 		}
 
 		if (bAllowSound)
@@ -1054,7 +1063,10 @@ void UCommonButtonBase::SetSelectedInternal(bool bInSelected, bool bAllowSound /
 	else
 	{
 		// Once deselected, restore the root button interactivity to the desired state
-		RootButton->SetInteractionEnabled(bInteractionEnabled);
+		if (RootButton.IsValid())
+		{
+			RootButton->SetInteractionEnabled(bInteractionEnabled);
+		}
 
 		NativeOnDeselected(bBroadcast);
 	}
@@ -2004,7 +2016,10 @@ void UCommonButtonBase::EnableButton()
 	if (!bButtonEnabled)
 	{
 		bButtonEnabled = true;
-		RootButton->SetButtonEnabled(true);
+		if (RootButton.IsValid())
+		{
+			RootButton->SetButtonEnabled(true);
+		}
 
 		SetButtonStyle();
 
@@ -2023,7 +2038,10 @@ void UCommonButtonBase::DisableButton()
 	if (bButtonEnabled)
 	{
 		bButtonEnabled = false;
-		RootButton->SetButtonEnabled(false);
+		if (RootButton.IsValid())
+		{
+			RootButton->SetButtonEnabled(false);
+		}
 
 		SetButtonStyle();
 
