Common UI:
- Add Throbber (Swidget) argument to SLoadGuard and setter for it. If unset, will use current ini defined throbber brush in an SImage
- Add throbber brush property to CommonLazyImage, CommonLazyWidget & CommonLoadGuard
- Use new brush property to feed the SWidget argument with an SImage when the brush is not set to draw as none

diff --git a/Engine/Plugins/Runtime/CommonUI/Source/CommonUI/Private/CommonLazyImage.cpp b/Engine/Plugins/Runtime/CommonUI/Source/CommonUI/Private/CommonLazyImage.cpp
--- a/Engine/Plugins/Runtime/CommonUI/Source/CommonUI/Private/CommonLazyImage.cpp
+++ b/Engine/Plugins/Runtime/CommonUI/Source/CommonUI/Private/CommonLazyImage.cpp
@@ -20,11 +20,13 @@ UCommonLazyImage::UCommonLazyImage(const FObjectInitializer& Initializer)
 	: Super(Initializer)
 {
 	LoadingBackgroundBrush.DrawAs = ESlateBrushDrawType::NoDrawType;
+	LoadingThrobberBrush.DrawAs = ESlateBrushDrawType::NoDrawType;
 }
 
 TSharedRef<SWidget> UCommonLazyImage::RebuildWidget()
 {
 	MyLoadGuard = SNew(SLoadGuard)
+		.Throbber(LoadingThrobberBrush.DrawAs != ESlateBrushDrawType::NoDrawType ? SNew(SImage).Image(&LoadingThrobberBrush) : TSharedPtr<SWidget>(nullptr))
 		.GuardBackgroundBrush(&LoadingBackgroundBrush)
 		.OnLoadingStateChanged_UObject(this, &UCommonLazyImage::HandleLoadGuardStateChanged)
 		[
@@ -51,6 +53,11 @@ void UCommonLazyImage::SynchronizeProperties()
 {
 	Super::SynchronizeProperties();
 
+	if (MyLoadGuard.IsValid())
+	{
+		MyLoadGuard->SetThrobber(LoadingThrobberBrush.DrawAs != ESlateBrushDrawType::NoDrawType ? SNew(SImage).Image(&LoadingThrobberBrush) : TSharedPtr<SWidget>(nullptr));
+	}
+
 #if WITH_EDITOR
 	if (IsDesignTime())
 	{
diff --git a/Engine/Plugins/Runtime/CommonUI/Source/CommonUI/Private/CommonLazyWidget.cpp b/Engine/Plugins/Runtime/CommonUI/Source/CommonUI/Private/CommonLazyWidget.cpp
--- a/Engine/Plugins/Runtime/CommonUI/Source/CommonUI/Private/CommonLazyWidget.cpp
+++ b/Engine/Plugins/Runtime/CommonUI/Source/CommonUI/Private/CommonLazyWidget.cpp
@@ -4,6 +4,7 @@
 #include "CommonActivatableWidget.h"
 #include "CommonWidgetPaletteCategories.h"
 #include "Engine/AssetManager.h"
+#include "Widgets/Images/SImage.h"
 
 #include UE_INLINE_GENERATED_CPP_BY_NAME(CommonLazyWidget)
 
@@ -11,12 +12,14 @@ UCommonLazyWidget::UCommonLazyWidget(const FObjectInitializer& Initializer)
 	: Super(Initializer)
 {
 	LoadingBackgroundBrush.DrawAs = ESlateBrushDrawType::NoDrawType;
+	LoadingThrobberBrush.DrawAs = ESlateBrushDrawType::NoDrawType;
 	SetVisibilityInternal(ESlateVisibility::SelfHitTestInvisible);
 }
 
 TSharedRef<SWidget> UCommonLazyWidget::RebuildWidget()
 {
 	MyLoadGuard = SNew(SLoadGuard)
+		.Throbber(LoadingThrobberBrush.DrawAs != ESlateBrushDrawType::NoDrawType ? SNew(SImage).Image(&LoadingThrobberBrush) : TSharedPtr<SWidget>(nullptr))
 		.GuardBackgroundBrush(&LoadingBackgroundBrush)
 		.OnLoadingStateChanged_UObject(this, &ThisClass::HandleLoadGuardStateChanged)
 		[
@@ -40,12 +43,17 @@ void UCommonLazyWidget::SynchronizeProperties()
 {
 	Super::SynchronizeProperties();
 
-#if WITH_EDITOR
-	if (IsDesignTime() && MyLoadGuard.IsValid())
+	if (MyLoadGuard.IsValid())
 	{
-		MyLoadGuard->SetForceShowSpinner(true);
-	}
+		MyLoadGuard->SetThrobber(LoadingThrobberBrush.DrawAs != ESlateBrushDrawType::NoDrawType ? SNew(SImage).Image(&LoadingThrobberBrush) : TSharedPtr<SWidget>(nullptr));
+
+#if WITH_EDITOR
+		if (IsDesignTime())
+		{
+			MyLoadGuard->SetForceShowSpinner(true);
+		}
 #endif
+	}
 }
 
 void UCommonLazyWidget::SetForceShowSpinner(bool bShowLoading)
diff --git a/Engine/Plugins/Runtime/CommonUI/Source/CommonUI/Private/CommonLoadGuard.cpp b/Engine/Plugins/Runtime/CommonUI/Source/CommonUI/Private/CommonLoadGuard.cpp
--- a/Engine/Plugins/Runtime/CommonUI/Source/CommonUI/Private/CommonLoadGuard.cpp
+++ b/Engine/Plugins/Runtime/CommonUI/Source/CommonUI/Private/CommonLoadGuard.cpp
@@ -9,7 +9,6 @@
 #include "Engine/StreamableManager.h"
 #include "Widgets/Text/STextBlock.h"
 #include "Widgets/Layout/SBorder.h"
-#include "Widgets/SBoxPanel.h"
 #include "Widgets/Images/SImage.h"
 #include "Widgets/Layout/SScaleBox.h"
 #include "ICommonUIModule.h"
@@ -40,6 +39,8 @@ void SLoadGuard::Construct(const FArguments& InArgs)
 		}
 	}
 
+	TSharedRef<SWidget> Throbber = InArgs._Throbber.IsValid() ? InArgs._Throbber.ToSharedRef() : SNew(SImage).Image(&ICommonUIModule::GetSettings().GetDefaultThrobberBrush());
+
 	ChildSlot
 	[
 		SNew(SOverlay)
@@ -75,13 +76,13 @@ void SLoadGuard::Construct(const FArguments& InArgs)
 				[
 					SNew(SHorizontalBox)
 					+ SHorizontalBox::Slot()
+					.Expose(ThrobberSlot)
 					.AutoWidth()
 					.HAlign(HAlign_Left)
 					.VAlign(VAlign_Center)
 					.Padding(0.f, 0.f, InArgs._GuardText.IsEmpty() ? 0.f : 18.f, 0.f)
 					[
-						SNew(SImage)
-						.Image(&ICommonUIModule::GetSettings().GetDefaultThrobberBrush())
+						Throbber
 					]
 
 					+ SHorizontalBox::Slot()
@@ -127,6 +128,14 @@ void SLoadGuard::SetThrobberHAlign(EHorizontalAlignment InHAlign)
 	GuardBorder->SetHAlign(InHAlign);
 }
 
+void SLoadGuard::SetThrobber(const TSharedPtr<SWidget>& InThrobber)
+{
+	if (ThrobberSlot)
+	{
+		ThrobberSlot->AttachWidget(InThrobber.IsValid() ? InThrobber.ToSharedRef() : SNew(SImage).Image(&ICommonUIModule::GetSettings().GetDefaultThrobberBrush()));
+	}
+}
+
 void SLoadGuard::SetGuardText(const FText& InText)
 {
 	GuardTextBlock->SetText(InText);
@@ -278,9 +287,9 @@ UCommonLoadGuard::UCommonLoadGuard(const FObjectInitializer& Initializer)
 {
 	SetVisibilityInternal(ESlateVisibility::SelfHitTestInvisible);
 
-	// Default to not showing the loading BG brush - it's an opt-in kind of thing
+	// Default to not showing the loading BG or throbber brushes - it's an opt-in kind of thing
 	LoadingBackgroundBrush.DrawAs = ESlateBrushDrawType::NoDrawType;
-
+	LoadingThrobberBrush.DrawAs = ESlateBrushDrawType::NoDrawType;
 }
 
 void UCommonLoadGuard::ReleaseSlateResources(bool bReleaseChildren)
@@ -353,7 +362,6 @@ void UCommonLoadGuard::PostLoad()
 	}
 
 #if WITH_EDITOR
-
 	// Needed for backwards compatibility of old style updating
 	if (!TextStyle && !bStyleNoLongerNeedsConversion && !IsRunningDedicatedServer())
 	{
@@ -372,6 +380,7 @@ TSharedRef<SWidget> UCommonLoadGuard::RebuildWidget()
 		.GuardText(LoadingText)
 		.GuardTextStyle(TextStyle)
 		.GuardBackgroundBrush(&LoadingBackgroundBrush)
+		.Throbber(LoadingThrobberBrush.DrawAs != ESlateBrushDrawType::NoDrawType ? SNew(SImage).Image(&LoadingThrobberBrush) : TSharedPtr<SWidget>(nullptr))
 		.ThrobberHAlign(ThrobberAlignment)
 		.OnLoadingStateChanged_UObject(this, &UCommonLoadGuard::HandleLoadingStateChanged);
 
@@ -395,7 +404,9 @@ void UCommonLoadGuard::SynchronizeProperties()
 			StyleCDO->ToTextBlockStyle(GuardTextStyle);
 			MyLoadGuard->SetGuardTextStyle(GuardTextStyle);
 		}
+
 		MyLoadGuard->SetGuardText(LoadingText);
+		MyLoadGuard->SetThrobber(LoadingThrobberBrush.DrawAs != ESlateBrushDrawType::NoDrawType ? SNew(SImage).Image(&LoadingThrobberBrush) : TSharedPtr<SWidget>(nullptr));
 		MyLoadGuard->SetThrobberHAlign(ThrobberAlignment);
 		MyLoadGuard->SetGuardBackgroundBrush(&LoadingBackgroundBrush);
 
diff --git a/Engine/Plugins/Runtime/CommonUI/Source/CommonUI/Public/CommonLazyImage.h b/Engine/Plugins/Runtime/CommonUI/Source/CommonUI/Public/CommonLazyImage.h
--- a/Engine/Plugins/Runtime/CommonUI/Source/CommonUI/Public/CommonLazyImage.h
+++ b/Engine/Plugins/Runtime/CommonUI/Source/CommonUI/Public/CommonLazyImage.h
@@ -92,6 +92,9 @@ private:
 	UPROPERTY(EditAnywhere, Category = Appearance)
 	FSlateBrush LoadingBackgroundBrush;
 
+	UPROPERTY(EditAnywhere, Category = Appearance)
+	FSlateBrush LoadingThrobberBrush;
+
 	/** 
 	 * If this image uses a material that a texture should be applied to, this is the name of the material param to use.
 	 * I.e. if this property is not blank, the resource object of our brush is a material, and we are given a lazy texture, that texture
diff --git a/Engine/Plugins/Runtime/CommonUI/Source/CommonUI/Public/CommonLazyWidget.h b/Engine/Plugins/Runtime/CommonUI/Source/CommonUI/Public/CommonLazyWidget.h
--- a/Engine/Plugins/Runtime/CommonUI/Source/CommonUI/Public/CommonLazyWidget.h
+++ b/Engine/Plugins/Runtime/CommonUI/Source/CommonUI/Public/CommonLazyWidget.h
@@ -57,6 +57,10 @@ private:
 	void RequestAsyncLoad(TSoftClassPtr<UObject> SoftObject, FStreamableDelegate DelegateToCall);
 	void HandleLoadGuardStateChanged(bool bIsLoading);
 
+	/** The loading throbber brush */
+	UPROPERTY(EditAnywhere, Category = Appearance)
+	FSlateBrush LoadingThrobberBrush;
+
 	UPROPERTY(EditAnywhere, Category = Appearance)
 	FSlateBrush LoadingBackgroundBrush;
 
diff --git a/Engine/Plugins/Runtime/CommonUI/Source/CommonUI/Public/CommonLoadGuard.h b/Engine/Plugins/Runtime/CommonUI/Source/CommonUI/Public/CommonLoadGuard.h
--- a/Engine/Plugins/Runtime/CommonUI/Source/CommonUI/Public/CommonLoadGuard.h
+++ b/Engine/Plugins/Runtime/CommonUI/Source/CommonUI/Public/CommonLoadGuard.h
@@ -4,6 +4,7 @@
 
 #include "Components/ContentWidget.h"
 #include "Widgets/Accessibility/SlateWidgetAccessibleTypes.h"
+#include "Widgets/SBoxPanel.h"
 #include "Widgets/SCompoundWidget.h"
 
 #include "CommonLoadGuard.generated.h"
@@ -26,12 +27,14 @@ class COMMONUI_API SLoadGuard : public SCompoundWidget
 {
 	SLATE_BEGIN_ARGS(SLoadGuard)
 		: _ThrobberHAlign(HAlign_Center)
+		, _Throbber(nullptr)
 		, _GuardTextStyle(nullptr)
 		, _GuardBackgroundBrush(nullptr)
 	{}
 		SLATE_DEFAULT_SLOT(FArguments, Content)
 	
 		SLATE_ARGUMENT(EHorizontalAlignment, ThrobberHAlign)
+		SLATE_ARGUMENT(TSharedPtr<SWidget>, Throbber)
 		SLATE_ARGUMENT(FText, GuardText)
 		SLATE_ARGUMENT(TSubclassOf<UCommonTextStyle>, GuardTextStyle)
 		SLATE_ARGUMENT(const FSlateBrush*, GuardBackgroundBrush)
@@ -49,6 +52,7 @@ public:
 	
 	void SetContent(const TSharedRef<SWidget>& InContent);
 	void SetThrobberHAlign(EHorizontalAlignment InHAlign);
+	void SetThrobber(const TSharedPtr<SWidget>& InThrobber);
 	void SetGuardText(const FText& InText);
 	void SetGuardTextStyle(const FTextBlockStyle& InGuardTextStyle);
 	void SetGuardBackgroundBrush(const FSlateBrush* InGuardBackground);
@@ -77,6 +81,7 @@ private:
 	TSharedPtr<SBorder> ContentBorder;
 	TSharedPtr<SBorder> GuardBorder;
 	TSharedPtr<STextBlock> GuardTextBlock;
+	SHorizontalBox::FSlot* ThrobberSlot = nullptr;
 
 	FOnLoadGuardStateChanged OnLoadingStateChanged;
 
@@ -215,11 +220,15 @@ private:
 	UPROPERTY(EditAnywhere, Category = LoadGuardThrobber)
 	FSlateBrush LoadingBackgroundBrush;
 
+	/** The loading throbber brush */
+	UPROPERTY(EditAnywhere, Category = LoadGuardThrobber)
+	FSlateBrush LoadingThrobberBrush;
+
 	/** The horizontal alignment of the loading throbber & message */
 	UPROPERTY(EditAnywhere, Category = LoadGuardThrobber)
 	TEnumAsByte<EHorizontalAlignment> ThrobberAlignment;
 
-	/** The horizontal alignment of the loading throbber & message */
+	/** The padding of the loading throbber & message */
 	UPROPERTY(EditAnywhere, Category = LoadGuardThrobber)
 	FMargin ThrobberPadding;
 
