[SButton, CommonButtonBase] OnClicked sound event added

#rnx

diff --git a/Engine/Plugins/Runtime/CommonUI/Source/CommonUI/Private/CommonButtonBase.cpp b/Engine/Plugins/Runtime/CommonUI/Source/CommonUI/Private/CommonButtonBase.cpp
--- a/Engine/Plugins/Runtime/CommonUI/Source/CommonUI/Private/CommonButtonBase.cpp
+++ b/Engine/Plugins/Runtime/CommonUI/Source/CommonUI/Private/CommonButtonBase.cpp
@@ -1751,6 +1751,7 @@ void UCommonButtonBase::BuildStyles()
 			SingleMaterialStyleMID = nullptr;
 		}
 		bool bHasPressedSlateSoundOverride = PressedSlateSoundOverride.GetResourceObject() != nullptr;
+		bool bHasClickedSlateSoundOverride = ClickedSlateSoundOverride.GetResourceObject() != nullptr;
 		bool bHasHoveredSlateSoundOverride = HoveredSlateSoundOverride.GetResourceObject() != nullptr;
 
 		NormalStyle.Normal = CommonButtonStyle->bSingleMaterial ? DynamicSingleMaterialBrush : CommonButtonStyle->NormalBase;
@@ -1762,6 +1763,7 @@ void UCommonButtonBase::BuildStyles()
 
 		// Sets the sound overrides for the Normal state
 		NormalStyle.PressedSlateSound = bHasPressedSlateSoundOverride ? PressedSlateSoundOverride : CommonButtonStyle->PressedSlateSound;
+		NormalStyle.ClickedSlateSound = bHasClickedSlateSoundOverride ? ClickedSlateSoundOverride : CommonButtonStyle->ClickedSlateSound;
 		NormalStyle.HoveredSlateSound = bHasHoveredSlateSoundOverride ? HoveredSlateSoundOverride : CommonButtonStyle->HoveredSlateSound;
 
 		SelectedStyle.Normal = CommonButtonStyle->bSingleMaterial ? DynamicSingleMaterialBrush : CommonButtonStyle->SelectedBase;
@@ -1789,6 +1791,18 @@ void UCommonButtonBase::BuildStyles()
 				? NormalStyle.PressedSlateSound
 				: CommonButtonStyle->SelectedPressedSlateSound.Sound;
 		}
+		
+		if (SelectedClickedSlateSoundOverride.GetResourceObject())
+		{
+			SelectedStyle.ClickedSlateSound = SelectedClickedSlateSoundOverride;
+		}
+		else
+		{
+			SelectedStyle.ClickedSlateSound =
+				bHasClickedSlateSoundOverride || !CommonButtonStyle->SelectedClickedSlateSound
+				? NormalStyle.ClickedSlateSound
+				: CommonButtonStyle->SelectedClickedSlateSound.Sound;
+		}
 
 		if (SelectedHoveredSlateSoundOverride.GetResourceObject())
 		{
@@ -1811,6 +1825,15 @@ void UCommonButtonBase::BuildStyles()
 				? LockedPressedSlateSoundOverride
 				: CommonButtonStyle->LockedPressedSlateSound.Sound;
 		}
+		
+		if (CommonButtonStyle->LockedClickedSlateSound || LockedClickedSlateSoundOverride.GetResourceObject())
+		{
+			LockedStyle.ClickedSlateSound =
+				LockedClickedSlateSoundOverride.GetResourceObject()
+				? LockedClickedSlateSoundOverride
+				: CommonButtonStyle->LockedClickedSlateSound.Sound;
+		}
+		
 		if (CommonButtonStyle->LockedHoveredSlateSound || LockedHoveredSlateSoundOverride.GetResourceObject())
 		{
 			LockedStyle.HoveredSlateSound =
@@ -1868,6 +1891,15 @@ void UCommonButtonBase::SetPressedSoundOverride(USoundBase* Sound)
 	}
 }
 
+void UCommonButtonBase::SetClickedSoundOverride(USoundBase* Sound)
+{
+	if (ClickedSlateSoundOverride.GetResourceObject() != Sound)
+	{
+		ClickedSlateSoundOverride.SetResourceObject(Sound);
+		BuildStyles();
+	}
+}
+
 void UCommonButtonBase::SetHoveredSoundOverride(USoundBase* Sound)
 {
 	if (HoveredSlateSoundOverride.GetResourceObject() != Sound)
@@ -1886,6 +1918,15 @@ void UCommonButtonBase::SetSelectedPressedSoundOverride(USoundBase* Sound)
 	}
 }
 
+void UCommonButtonBase::SetSelectedClickedSoundOverride(USoundBase* Sound)
+{
+	if (SelectedClickedSlateSoundOverride.GetResourceObject() != Sound)
+	{
+		SelectedClickedSlateSoundOverride.SetResourceObject(Sound);
+		BuildStyles();
+	}
+}
+
 void UCommonButtonBase::SetSelectedHoveredSoundOverride(USoundBase* Sound)
 {
 	if (SelectedHoveredSlateSoundOverride.GetResourceObject() != Sound)
@@ -1904,6 +1945,15 @@ void UCommonButtonBase::SetLockedPressedSoundOverride(USoundBase* Sound)
 	}
 }
 
+void UCommonButtonBase::SetLockedClickedSoundOverride(USoundBase* Sound)
+{
+	if (LockedClickedSlateSoundOverride.GetResourceObject() != Sound)
+	{
+		LockedClickedSlateSoundOverride.SetResourceObject(Sound);
+		BuildStyles();
+	}
+}
+
 void UCommonButtonBase::SetLockedHoveredSoundOverride(USoundBase* Sound)
 {
 	if (LockedHoveredSlateSoundOverride.GetResourceObject() != Sound)
@@ -1923,6 +1973,16 @@ void UCommonButtonBase::SetPressedSlateSoundOverride(const FSlateSound& InPresse
 	}
 }
 
+void UCommonButtonBase::SetClickedSlateSoundOverride(const FSlateSound& InClickedSlateSoundOverride)
+{
+	const bool bBuildStyles = InClickedSlateSoundOverride.GetResourceObject() != ClickedSlateSoundOverride.GetResourceObject();
+	ClickedSlateSoundOverride = InClickedSlateSoundOverride;
+	if (bBuildStyles)
+	{
+		BuildStyles();
+	}
+}
+
 void UCommonButtonBase::SetHoveredSlateSoundOverride(const FSlateSound& InHoveredSlateSoundOverride)
 {
 	const bool bBuildStyles = InHoveredSlateSoundOverride.GetResourceObject() != HoveredSlateSoundOverride.GetResourceObject();
@@ -1943,6 +2003,16 @@ void UCommonButtonBase::SetSelectedPressedSlateSoundOverride(const FSlateSound&
 	}
 }
 
+void UCommonButtonBase::SetSelectedClickedSlateSoundOverride(const FSlateSound& InSelectedClickedSlateSoundOverride)
+{
+	const bool bBuildStyles = InSelectedClickedSlateSoundOverride.GetResourceObject() != SelectedClickedSlateSoundOverride.GetResourceObject();
+	SelectedClickedSlateSoundOverride = InSelectedClickedSlateSoundOverride;
+	if (bBuildStyles)
+	{
+		BuildStyles();
+	}
+}
+
 void UCommonButtonBase::SetSelectedHoveredSlateSoundOverride(const FSlateSound& InSelectedHoveredSlateSoundOverride)
 {
 	const bool bBuildStyles = InSelectedHoveredSlateSoundOverride.GetResourceObject() != SelectedHoveredSlateSoundOverride.GetResourceObject();
@@ -1963,6 +2033,16 @@ void UCommonButtonBase::SetLockedPressedSlateSoundOverride(const FSlateSound& In
 	}
 }
 
+void UCommonButtonBase::SetLockedClickedSlateSoundOverride(const FSlateSound& InLockedClickedSlateSoundOverride)
+{
+	const bool bBuildStyles = InLockedClickedSlateSoundOverride.GetResourceObject() != LockedClickedSlateSoundOverride.GetResourceObject();
+	LockedClickedSlateSoundOverride = InLockedClickedSlateSoundOverride;
+	if (bBuildStyles)
+	{
+		BuildStyles();
+	}
+}
+
 void UCommonButtonBase::SetLockedHoveredSlateSoundOverride(const FSlateSound& InLockedHoveredSlateSoundOverride)
 {
 	const bool bBuildStyles = InLockedHoveredSlateSoundOverride.GetResourceObject() != LockedHoveredSlateSoundOverride.GetResourceObject();
diff --git a/Engine/Plugins/Runtime/CommonUI/Source/CommonUI/Public/CommonButtonBase.h b/Engine/Plugins/Runtime/CommonUI/Source/CommonUI/Public/CommonButtonBase.h
--- a/Engine/Plugins/Runtime/CommonUI/Source/CommonUI/Public/CommonButtonBase.h
+++ b/Engine/Plugins/Runtime/CommonUI/Source/CommonUI/Public/CommonButtonBase.h
@@ -157,13 +157,25 @@ public:
 	UPROPERTY(EditDefaultsOnly, BlueprintReadOnly, Category = "Properties", meta = (DisplayName = "Pressed Sound"))
 	FSlateSound PressedSlateSound;
 
+	/** The sound to play when the button is clicked */
+	UPROPERTY(EditDefaultsOnly, BlueprintReadOnly, Category = "Properties", meta = (DisplayName = "Clicked Sound"))
+	FSlateSound ClickedSlateSound;
+
 	/** The sound to play when the button is pressed while selected */
 	UPROPERTY(EditDefaultsOnly, BlueprintReadOnly, Category = "Properties", meta = (DisplayName = "Selected Pressed Sound"))
 	FCommonButtonStyleOptionalSlateSound SelectedPressedSlateSound;
 
+	/** The sound to play when the button is clicked while selected */
+	UPROPERTY(EditDefaultsOnly, BlueprintReadOnly, Category = "Properties", meta = (DisplayName = "Selected Clicked Sound"))
+	FCommonButtonStyleOptionalSlateSound SelectedClickedSlateSound;
+
 	/** The sound to play when the button is pressed while locked */
 	UPROPERTY(EditDefaultsOnly, BlueprintReadOnly, Category = "Properties", meta = (DisplayName = "Locked Pressed Sound"))
 	FCommonButtonStyleOptionalSlateSound LockedPressedSlateSound;
+
+	/** The sound to play when the button is clicked while locked */
+	UPROPERTY(EditDefaultsOnly, BlueprintReadOnly, Category = "Properties", meta = (DisplayName = "Locked Clicked Sound"))
+	FCommonButtonStyleOptionalSlateSound LockedClickedSlateSound;
 	
 	/** The sound to play when the button is hovered */
 	UPROPERTY(EditDefaultsOnly, BlueprintReadOnly, Category = "Properties", meta = (DisplayName = "Hovered Sound"))
@@ -495,17 +507,26 @@ public:
 	UFUNCTION(BlueprintCallable, Category = "Common Button|Sound")
 	COMMONUI_API void SetPressedSoundOverride(USoundBase* Sound);
 
+	UFUNCTION(BlueprintCallable, Category = "Common Button|Sound")
+	COMMONUI_API void SetClickedSoundOverride(USoundBase* Sound);
+
 	UFUNCTION(BlueprintCallable, Category = "Common Button|Sound")
 	COMMONUI_API void SetHoveredSoundOverride(USoundBase* Sound);
 
 	UFUNCTION(BlueprintCallable, Category = "Common Button|Sound")
 	COMMONUI_API void SetSelectedPressedSoundOverride(USoundBase* Sound);
 
+	UFUNCTION(BlueprintCallable, Category = "Common Button|Sound")
+	COMMONUI_API void SetSelectedClickedSoundOverride(USoundBase* Sound);
+
 	UFUNCTION(BlueprintCallable, Category = "Common Button|Sound")
 	COMMONUI_API void SetSelectedHoveredSoundOverride(USoundBase* Sound);
 	
 	UFUNCTION(BlueprintCallable, Category = "Common Button|Sound")
 	COMMONUI_API void SetLockedPressedSoundOverride(USoundBase* Sound);
+	
+	UFUNCTION(BlueprintCallable, Category = "Common Button|Sound")
+	COMMONUI_API void SetLockedClickedSoundOverride(USoundBase* Sound);
 
 	UFUNCTION(BlueprintCallable, Category = "Common Button|Sound")
 	COMMONUI_API void SetLockedHoveredSoundOverride(USoundBase* Sound);
@@ -546,10 +567,13 @@ protected:
 
 	// In some scenarios, USoundBase* setters are bypassed, so standard setters must also be provided for sound override properties
 	COMMONUI_API void SetPressedSlateSoundOverride(const FSlateSound& InPressedSlateSoundOverride);
+	COMMONUI_API void SetClickedSlateSoundOverride(const FSlateSound& InClickedSlateSoundOverride);
 	COMMONUI_API void SetHoveredSlateSoundOverride(const FSlateSound& InHoveredSlateSoundOverride);
 	COMMONUI_API void SetSelectedPressedSlateSoundOverride(const FSlateSound& InSelectedPressedSlateSoundOverride);
+	COMMONUI_API void SetSelectedClickedSlateSoundOverride(const FSlateSound& InSelectedClickedSlateSoundOverride);
 	COMMONUI_API void SetSelectedHoveredSlateSoundOverride(const FSlateSound& InSelectedHoveredSlateSoundOverride);
 	COMMONUI_API void SetLockedPressedSlateSoundOverride(const FSlateSound& InLockedPressedSlateSoundOverride);
+	COMMONUI_API void SetLockedClickedSlateSoundOverride(const FSlateSound& InLockedClickedSlateSoundOverride);
 	COMMONUI_API void SetLockedHoveredSlateSoundOverride(const FSlateSound& InLockedHoveredSlateSoundOverride);
 
 	/** Helper function to bind to input method change events */
@@ -767,6 +791,12 @@ protected:
 	UPROPERTY(EditAnywhere, BlueprintReadOnly, Setter, Category = Sound, meta = (DisplayName = "Pressed Sound Override"))
 	FSlateSound PressedSlateSoundOverride;
 
+	/**
+	 * Optional override for the sound to play when this button is clicked (based on Click/Touch/Press methods).
+	 */
+	UPROPERTY(EditAnywhere, BlueprintReadOnly, Setter, Category = Sound, meta = (DisplayName = "Clicked Sound Override"))
+	FSlateSound ClickedSlateSoundOverride;
+
 	/**
 	 * Optional override for the sound to play when this button is hovered.
 	 * Also used for the Selected and Locked Hovered state if their respective overrides are empty.
@@ -778,6 +808,10 @@ protected:
 	UPROPERTY(EditAnywhere, BlueprintReadOnly, Setter, Category = Sound, meta = (DisplayName = "Selected Pressed Sound Override"))
 	FSlateSound SelectedPressedSlateSoundOverride;
 
+	/** Optional override for the sound to play when this button is clicked while Selected */
+	UPROPERTY(EditAnywhere, BlueprintReadOnly, Setter, Category = Sound, meta = (DisplayName = "Selected Clicked Sound Override"))
+	FSlateSound SelectedClickedSlateSoundOverride;
+
 	/** Optional override for the sound to play when this button is hovered while Selected */
 	UPROPERTY(EditAnywhere, BlueprintReadOnly, Setter, Category = Sound, meta = (DisplayName = "Selected Hovered Sound Override"))
 	FSlateSound SelectedHoveredSlateSoundOverride;
@@ -786,6 +820,10 @@ protected:
 	UPROPERTY(EditAnywhere, BlueprintReadOnly, Setter, Category = Sound, meta = (DisplayName = "Locked Pressed Sound Override"))
 	FSlateSound LockedPressedSlateSoundOverride;
 
+	/** Optional override for the sound to play when this button is clicked while Locked */
+	UPROPERTY(EditAnywhere, BlueprintReadOnly, Setter, Category = Sound, meta = (DisplayName = "Locked Clicked Sound Override"))
+	FSlateSound LockedClickedSlateSoundOverride;
+
 	/** Optional override for the sound to play when this button is hovered while Locked */
 	UPROPERTY(EditAnywhere, BlueprintReadOnly, Setter, Category = Sound, meta = (DisplayName = "Locked Hovered Sound Override"))
 	FSlateSound LockedHoveredSlateSoundOverride;
diff --git a/Engine/Source/Runtime/Slate/Private/Widgets/Input/SButton.cpp b/Engine/Source/Runtime/Slate/Private/Widgets/Input/SButton.cpp
--- a/Engine/Source/Runtime/Slate/Private/Widgets/Input/SButton.cpp
+++ b/Engine/Source/Runtime/Slate/Private/Widgets/Input/SButton.cpp
@@ -114,6 +114,7 @@ void SButton::Construct( const FArguments& InArgs )
 
 	HoveredSound = InArgs._HoveredSoundOverride.Get(InArgs._ButtonStyle->HoveredSlateSound);
 	PressedSound = InArgs._PressedSoundOverride.Get(InArgs._ButtonStyle->PressedSlateSound);
+	ClickedSound = InArgs._ClickedSoundOverride.Get(InArgs._ButtonStyle->ClickedSlateSound);
 
 	// Text overrides button content. If nothing is specified, put an null widget in the button.
 	// Null content makes the button enter a special mode where it will ask to be as big as the image used for its border.
@@ -488,6 +489,7 @@ void SButton::OnMouseCaptureLost(const FCaptureLostEvent& CaptureLostEvent)
 
 FReply SButton::ExecuteOnClick()
 {
+	PlayClickedSound();
 	if (OnClicked.IsBound())
 	{
 		FReply Reply = OnClicked.Execute();
@@ -564,6 +566,11 @@ void SButton::PlayPressedSound() const
 	FSlateApplication::Get().PlaySound( PressedSound );
 }
 
+void SButton::PlayClickedSound() const
+{
+	FSlateApplication::Get().PlaySound( ClickedSound );
+}
+
 void SButton::PlayHoverSound() const
 {
 	FSlateApplication::Get().PlaySound( HoveredSound );
@@ -598,6 +605,11 @@ void SButton::SetPressedSound(TOptional<FSlateSound> InPressedSound)
 	PressedSound = InPressedSound.Get(Style->PressedSlateSound);
 }
 
+void SButton::SetClickedSound(TOptional<FSlateSound> InClickedSound)
+{
+	ClickedSound = InClickedSound.Get(Style->ClickedSlateSound);
+}
+
 void SButton::SetOnClicked(FOnClicked InOnClicked)
 {
 	OnClicked = InOnClicked;
@@ -652,6 +664,7 @@ void SButton::SetButtonStyle(const FButtonStyle* InButtonStyle)
 
 	HoveredSound = Style->HoveredSlateSound;
 	PressedSound = Style->PressedSlateSound;
+	ClickedSound = Style->ClickedSlateSound;
 
 	UpdatePadding();
 	UpdateShowDisabledEffect();
diff --git a/Engine/Source/Runtime/Slate/Public/Widgets/Input/SButton.h b/Engine/Source/Runtime/Slate/Public/Widgets/Input/SButton.h
--- a/Engine/Source/Runtime/Slate/Public/Widgets/Input/SButton.h
+++ b/Engine/Source/Runtime/Slate/Public/Widgets/Input/SButton.h
@@ -115,6 +115,9 @@ public:
 		/** The sound to play when the button is pressed */
 		SLATE_ARGUMENT( TOptional<FSlateSound>, PressedSoundOverride )
 
+		/** The sound to play when the button is clicked */
+		SLATE_ARGUMENT( TOptional<FSlateSound>, ClickedSoundOverride )
+
 		/** The sound to play when the button is hovered */
 		SLATE_ARGUMENT( TOptional<FSlateSound>, HoveredSoundOverride )
 
@@ -168,6 +171,9 @@ public:
 	/** See PressedSound attribute */
 	SLATE_API void SetPressedSound(TOptional<FSlateSound> InPressedSound);
 
+	/** See ClickedSound attribute */
+	SLATE_API void SetClickedSound(TOptional<FSlateSound> InClickedSound);
+
 	/** See OnClicked event */
 	SLATE_API void SetOnClicked(FOnClicked InOnClicked);
 
@@ -236,6 +242,9 @@ protected:
 	/** Play the pressed sound */
 	SLATE_API void PlayPressedSound() const;
 
+	/** Play the clicked sound */
+	SLATE_API void PlayClickedSound() const;
+
 	/** Play the hovered sound */
 	SLATE_API void PlayHoverSound() const;
 
@@ -299,6 +308,9 @@ private:
 	/** The Sound to play when the button is pressed */
 	FSlateSound PressedSound;
 
+	/** The Sound to play when the button is clicked */
+	FSlateSound ClickedSound;
+
 	/** Sets whether a click should be triggered on mouse down, mouse up, or that both a mouse down and up are required. */
 	TEnumAsByte<EButtonClickMethod::Type> ClickMethod;
 
diff --git a/Engine/Source/Runtime/SlateCore/Public/Styling/SlateTypes.h b/Engine/Source/Runtime/SlateCore/Public/Styling/SlateTypes.h
--- a/Engine/Source/Runtime/SlateCore/Public/Styling/SlateTypes.h
+++ b/Engine/Source/Runtime/SlateCore/Public/Styling/SlateTypes.h
@@ -586,6 +586,13 @@ struct FButtonStyle : public FSlateWidgetStyle
 	FSlateSound PressedSlateSound;
 	FButtonStyle& SetPressedSound( const FSlateSound& InPressedSound ){ PressedSlateSound = InPressedSound; return *this; }
 
+	/**
+	 * The sound the button should play when clicked
+	 */	
+	UPROPERTY(EditAnywhere, BlueprintReadWrite, Category=Sound, meta=( DisplayName="Clicked Sound" ))
+	FSlateSound ClickedSlateSound;
+	FButtonStyle& SetClickedSound( const FSlateSound& InClickedSound ){ ClickedSlateSound = InClickedSound; return *this; }
+
 	/**
 	 * The sound the button should play when initially hovered over
 	 */	
