Added a function to set the scroll orientation of Common Text.

diff --git a/Engine/Plugins/Runtime/CommonUI/Source/CommonUI/Private/CommonTextBlock.cpp b/Engine/Plugins/Runtime/CommonUI/Source/CommonUI/Private/CommonTextBlock.cpp
--- a/Engine/Plugins/Runtime/CommonUI/Source/CommonUI/Private/CommonTextBlock.cpp
+++ b/Engine/Plugins/Runtime/CommonUI/Source/CommonUI/Private/CommonTextBlock.cpp
@@ -374,6 +374,12 @@ void UCommonTextBlock::SetStyle(TSubclassOf<UCommonTextStyle> InStyle)
 	SynchronizeProperties();
 }
 
+void UCommonTextBlock::SetScrollOrientation(TEnumAsByte<EOrientation> InScrollOrientation)
+{
+	ScrollOrientation = InScrollOrientation;
+	SynchronizeProperties();
+}
+
 const FMargin& UCommonTextBlock::GetMargin()
 {
 	return Margin;
diff --git a/Engine/Plugins/Runtime/CommonUI/Source/CommonUI/Public/CommonTextBlock.h b/Engine/Plugins/Runtime/CommonUI/Source/CommonUI/Public/CommonTextBlock.h
--- a/Engine/Plugins/Runtime/CommonUI/Source/CommonUI/Public/CommonTextBlock.h
+++ b/Engine/Plugins/Runtime/CommonUI/Source/CommonUI/Public/CommonTextBlock.h
@@ -150,6 +150,9 @@ public:
 	UFUNCTION(BlueprintCallable, Category = "Common Text")
 	UE_API void SetStyle(TSubclassOf<UCommonTextStyle> InStyle);
 
+	UFUNCTION(BlueprintCallable, Category = "Common Text")
+	UE_API void SetScrollOrientation(TEnumAsByte<EOrientation> InScrollOrientation);
+
 	UFUNCTION(BlueprintCallable, Category = "Common Text")
 	UE_API const FMargin& GetMargin();
 
