SLoadGuard:
- Fix OnLoadingStateChanged not getting copied in the Construct call

diff --git a/Engine/Plugins/Runtime/CommonUI/Source/CommonUI/Private/CommonLoadGuard.cpp b/Engine/Plugins/Runtime/CommonUI/Source/CommonUI/Private/CommonLoadGuard.cpp
--- a/Engine/Plugins/Runtime/CommonUI/Source/CommonUI/Private/CommonLoadGuard.cpp
+++ b/Engine/Plugins/Runtime/CommonUI/Source/CommonUI/Private/CommonLoadGuard.cpp
@@ -39,6 +39,8 @@ void SLoadGuard::Construct(const FArguments& InArgs)
 		}
 	}
 
+	OnLoadingStateChanged = InArgs._OnLoadingStateChanged;
+
 	TSharedRef<SWidget> Throbber = InArgs._Throbber.IsValid() ? InArgs._Throbber.ToSharedRef() : SNew(SImage).Image(&ICommonUIModule::GetSettings().GetDefaultThrobberBrush());
 
 	ChildSlot
