[CommonUI][UI] Fixing the function that checks if a Actions is valid for a specific InputType

If one of the action inputs become unbind, this one will be recognized as Keyboard Key, since the checks
for IsGamepadKey and IsTouchKey will return false, and this will make this action be considered valid for Keyboard.

To fix this I added a check that will consider a key invalid if it is unbind.

#rnx

diff --git a/Engine/Plugins/Runtime/CommonUI/Source/CommonUI/Private/CommonUITypes.cpp b/Engine/Plugins/Runtime/CommonUI/Source/CommonUI/Private/CommonUITypes.cpp
--- a/Engine/Plugins/Runtime/CommonUI/Source/CommonUI/Private/CommonUITypes.cpp
+++ b/Engine/Plugins/Runtime/CommonUI/Source/CommonUI/Private/CommonUITypes.cpp
@@ -411,6 +411,11 @@ bool CommonUI::ActionValidForInputType(const ULocalPlayer* LocalPlayer, ECommonI
 
 bool CommonUI::IsKeyValidForInputType(const FKey& Key, ECommonInputType InputType)
 {
+	if (!Key.IsValid())
+	{
+		return false;
+	}
+
 	switch (InputType)
 	{
 	case ECommonInputType::MouseAndKeyboard:
