CommonUI: Improve the ensure messaging around invalid actions in Common Bound Action Bar

#rnx

diff --git a/Engine/Plugins/Runtime/CommonUI/Source/CommonUI/Private/Input/CommonBoundActionBar.cpp b/Engine/Plugins/Runtime/CommonUI/Source/CommonUI/Private/Input/CommonBoundActionBar.cpp
--- a/Engine/Plugins/Runtime/CommonUI/Source/CommonUI/Private/Input/CommonBoundActionBar.cpp
+++ b/Engine/Plugins/Runtime/CommonUI/Source/CommonUI/Private/Input/CommonBoundActionBar.cpp
@@ -287,7 +287,8 @@ void UCommonBoundActionBar::HandleDeferredDisplayUpdate()
 							bool bIsValidActionA = LegacyDataA || InputActionA;
 							bool bIsValidActionB = LegacyDataB || InputActionB;
 
-							if (ensureMsgf((bIsValidActionA && bIsValidActionB), TEXT("Action bindings not displayed yet -- array filter enforces they are not included")))
+							if (ensureMsgf(bIsValidActionA, TEXT("Binding is invalid: %s"), *BindingA->ToDebugString())
+								&& ensureMsgf(bIsValidActionB, TEXT("Binding is invalid: %s"), *BindingB->ToDebugString()))
 							{
 								bool bAIsBack = IsKeyBackAction(LegacyDataA, InputActionA);
 								bool bBIsBack = IsKeyBackAction(LegacyDataB, InputActionB);
