Retain focused on hovered widget when switching input method to gamepad
Require Cvar `CommonUI.GamepadFocusHoveredWidget` is enabled

diff --git a/Engine/Plugins/Runtime/CommonUI/Source/CommonUI/Private/Input/UIActionRouterTypes.cpp b/Engine/Plugins/Runtime/CommonUI/Source/CommonUI/Private/Input/UIActionRouterTypes.cpp
--- a/Engine/Plugins/Runtime/CommonUI/Source/CommonUI/Private/Input/UIActionRouterTypes.cpp
+++ b/Engine/Plugins/Runtime/CommonUI/Source/CommonUI/Private/Input/UIActionRouterTypes.cpp
@@ -23,6 +23,12 @@
 
 DEFINE_LOG_CATEGORY(LogUIActionRouter);
 
+static const TAutoConsoleVariable<bool> CVarGamepadFocusHoveredWidget(
+	TEXT("CommonUI.GamepadFocusHoveredWidget"),
+	false,
+	TEXT("If enabled, the mouse hovered widget, will be focused when switching to gamepad input"));
+
+
 const TCHAR* InputEventToString(EInputEvent InputEvent)
 {
 	switch (InputEvent)
@@ -1506,7 +1512,11 @@ void FActivatableTreeRoot::HandleInputMethodChanged(ECommonInputType InputMethod
 {
 	if (IsReceivingInput() && LeafmostActiveNode.IsValid() && ensure(LeafmostActiveNode.Pin()->IsReceivingInput()))
 	{
-		ApplyLeafmostNodeConfig();
+		const bool bRetainFocus = (InputMethod == ECommonInputType::Gamepad) && CVarGamepadFocusHoveredWidget.GetValueOnAnyThread();
+		
+		// If focus was not retained, let leafmost node determine which widget should be focused
+		ApplyLeafmostNodeConfig(bRetainFocus);
+
 		if (InputMethod != ECommonInputType::Gamepad)
 		{
 			LeafmostActiveNode.Pin()->CacheFocusRestorationTarget();
@@ -1514,7 +1524,7 @@ void FActivatableTreeRoot::HandleInputMethodChanged(ECommonInputType InputMethod
 	}
 }
 
-void FActivatableTreeRoot::ApplyLeafmostNodeConfig()
+void FActivatableTreeRoot::ApplyLeafmostNodeConfig(bool bAttemptRetainFocus)
 {
 #if WITH_EDITOR
 	if (!IsViewportWindowInFocusPath(GetActionRouter()))
@@ -1528,6 +1538,7 @@ void FActivatableTreeRoot::ApplyLeafmostNodeConfig()
 		return;
 	}
 
+
 	if (FActivatableTreeNodePtr PinnedLeafmostNode = LeafmostActiveNode.Pin())
 	{
 		GetActionRouter().SetActiveActivationMetadata(PinnedLeafmostNode->FindActivationMetadata());
@@ -1547,7 +1558,17 @@ void FActivatableTreeRoot::ApplyLeafmostNodeConfig()
 				GetActionRouter().SetActiveUIInputConfig(FUIInputConfig());
 			}
 
-			FocusLeafmostNode();
+			bool bFocusLeafMost = true;
+			if (bAttemptRetainFocus)
+			{
+				// Focus if could not retain focus
+				bFocusLeafMost = !GamepadFocusHoveredWidget();
+			}
+			
+			if(bFocusLeafMost)
+			{ 
+				FocusLeafmostNode();
+			}
 		}
 		else
 		{
@@ -1680,6 +1701,32 @@ void FActivatableTreeRoot::HandleRequestRefreshLeafmostFocus()
 	}
 }
 
+bool FActivatableTreeRoot::GamepadFocusHoveredWidget()
+{
+
+	FSlateApplication& SlateApplication = FSlateApplication::Get();
+	TSharedPtr<FActivatableTreeNode> PinnedLeafmostActiveNode = LeafmostActiveNode.Pin();
+
+	// First attempt to regain focus by checking if any activable widgets are under the last mouse cursor position
+	FWidgetPath WidgetPath = SlateApplication.LocateWindowUnderMouse(SlateApplication.GetLastCursorPos(), SlateApplication.GetTopLevelWindows());
+	for (int Index = WidgetPath.Widgets.Num() - 1; Index >= 0; --Index)
+	{
+		TSharedRef<SWidget> WidgetUnderCursor = WidgetPath.Widgets[Index].Widget;
+		UCommonActivatableWidget* ActivatableWidgetUnderCursor = UCommonUIActionRouterBase::FindActivatable(WidgetUnderCursor , GetActionRouter().GetLocalPlayerChecked());
+		if (ActivatableWidgetUnderCursor != nullptr && IsWidgetInNodeHierarchy(ActivatableWidgetUnderCursor, *PinnedLeafmostActiveNode))
+		{
+			if (FSlateApplication::Get().SetUserFocus(GetOwnerUserIndex(), WidgetUnderCursor))
+			{
+				ULocalPlayer& LocalPlayer = *GetActionRouter().GetLocalPlayerChecked();
+				LocalPlayer.GetSlateOperations().CancelFocusRequest();
+			}
+			return true;
+		}
+	}
+
+	return false;
+}
+
 void FActivatableTreeRoot::DebugDump(FString& OutputStr, bool bIncludeActions, bool bIncludeChildren, bool bIncludeInactive) const
 {
 	DebugDumpRecursive(OutputStr, 0, bIncludeActions, bIncludeChildren, bIncludeInactive);
diff --git a/Engine/Plugins/Runtime/CommonUI/Source/CommonUI/Private/Input/UIActionRouterTypes.h b/Engine/Plugins/Runtime/CommonUI/Source/CommonUI/Private/Input/UIActionRouterTypes.h
--- a/Engine/Plugins/Runtime/CommonUI/Source/CommonUI/Private/Input/UIActionRouterTypes.h
+++ b/Engine/Plugins/Runtime/CommonUI/Source/CommonUI/Private/Input/UIActionRouterTypes.h
@@ -205,7 +205,7 @@ public:
 
 	UE_API void RefreshCachedRestorationTarget();
 
-	UE_API void ApplyLeafmostNodeConfig();
+	UE_API void ApplyLeafmostNodeConfig(bool bAttemptRetainFocus = false);
 
 	UE_API bool IsAnActionDomainRoot() const;
 	UE_API bool IsActiveActionDomainRoot() const;
@@ -223,6 +223,8 @@ private:
 
 	void HandleRequestRefreshLeafmostFocus();
 
+	bool GamepadFocusHoveredWidget();
+
 	// WeakPtr because the root itself can be the primary active node - results in a circular ref leak using a full SharedPtr here
 	TWeakPtr<FActivatableTreeNode> LeafmostActiveNode;
 };
