[CommonInput] Make UCommonInputSubsystem::SetCurrentInputType try to recalculate input type if last call was locked by thrashing

- This change fixes an issue where current input type is never updated and propagated to listeners if SetCurrentInputType is called during input method thrashing.
- For example when holding down a keyboard button and pushing an analog stick, then during thrashing lock releasing the keyboard button while continuing to use the controller it is possible for CurrentInputType to be stuck as MouseAndKeyboard even though a controller is used.
- The issue can be observed when using a UCommonActionWidget with icons set up for both PC and console.

diff --git a/Engine/Plugins/Runtime/CommonUI/Source/CommonInput/Private/CommonInputSubsystem.cpp b/Engine/Plugins/Runtime/CommonUI/Source/CommonInput/Private/CommonInputSubsystem.cpp
--- a/Engine/Plugins/Runtime/CommonUI/Source/CommonInput/Private/CommonInputSubsystem.cpp
+++ b/Engine/Plugins/Runtime/CommonUI/Source/CommonInput/Private/CommonInputSubsystem.cpp
@@ -305,12 +305,12 @@ void UCommonInputSubsystem::RecalculateCurrentInputType()
 
 void UCommonInputSubsystem::SetCurrentInputType(ECommonInputType NewInputType)
 {
-	if ((RawInputType != NewInputType) && PlatformSupportsInputType(NewInputType))
+	if (((RawInputType != NewInputType) || bInputMethodLockedByThrashing) && PlatformSupportsInputType(NewInputType))
 	{
 		RawInputType = NewInputType;
 
-		const bool bIsLockedByThrashing = CheckForInputMethodThrashing(NewInputType);
-		if (!bIsLockedByThrashing)
+		bInputMethodLockedByThrashing = CheckForInputMethodThrashing(NewInputType);
+		if (!bInputMethodLockedByThrashing)
 		{
 			RecalculateCurrentInputType();
 		}
diff --git a/Engine/Plugins/Runtime/CommonUI/Source/CommonInput/Public/CommonInputSubsystem.h b/Engine/Plugins/Runtime/CommonUI/Source/CommonInput/Public/CommonInputSubsystem.h
--- a/Engine/Plugins/Runtime/CommonUI/Source/CommonInput/Public/CommonInputSubsystem.h
+++ b/Engine/Plugins/Runtime/CommonUI/Source/CommonInput/Public/CommonInputSubsystem.h
@@ -152,6 +152,10 @@ private:
 	UPROPERTY(Transient)
 	FName GamepadInputType;
 
+	/** Whether the last input type recalculation was locked out by input thrashing */
+	UPROPERTY(Transient)
+	bool bInputMethodLockedByThrashing = false;
+
 	/**  */
 	UPROPERTY(Transient)
 	TMap<FName, ECommonInputType> CurrentInputLocks;
