PR #13106: CommonUI - BoundActionButton - Actions requiring hold can now require it for button mouse presses as well.

diff --git a/Engine/Plugins/Runtime/CommonUI/Source/CommonUI/Private/Input/CommonBoundActionButton.cpp b/Engine/Plugins/Runtime/CommonUI/Source/CommonUI/Private/Input/CommonBoundActionButton.cpp
--- a/Engine/Plugins/Runtime/CommonUI/Source/CommonUI/Private/Input/CommonBoundActionButton.cpp
+++ b/Engine/Plugins/Runtime/CommonUI/Source/CommonUI/Private/Input/CommonBoundActionButton.cpp
@@ -4,6 +4,7 @@
 #include "CommonTextBlock.h"
 #include "Input/UIActionBinding.h"
 #include "CommonActionWidget.h"
+#include "CommonUITypes.h"
 #include "Framework/Application/SlateApplication.h"
 #include "Internationalization/LocKeyFuncs.h"
 
@@ -24,6 +25,12 @@ void UCommonBoundActionButton::SetRepresentedAction(FUIActionBindingHandle InBin
 	if (TSharedPtr<FUIActionBinding> NewBinding = FUIActionBinding::FindBinding(InBindingHandle))
 	{
 		NewBinding->OnHoldActionProgressed.AddUObject(this, &UCommonBoundActionButton::NativeOnActionProgress);
+
+		if (bLinkRequiresHoldToBindingHold)
+		{
+			const FCommonInputActionDataBase* InputActionData = NewBinding->GetLegacyInputActionData();
+			SetRequiresHold(InputActionData ? InputActionData->HasHoldBindings() : false);
+		}
 	}
 }
 
@@ -79,4 +86,24 @@ void UCommonBoundActionButton::UpdateInputActionWidget()
 	}
 }
 
+void UCommonBoundActionButton::UpdateHoldData(ECommonInputType CurrentInputType)
+{
+	if (bLinkRequiresHoldToBindingHold)
+	{
+		if (const TSharedPtr<FUIActionBinding> ActionBinding = FUIActionBinding::FindBinding(BindingHandle))
+		{
+			if (const FCommonInputActionDataBase* DataTableRow = ActionBinding->GetLegacyInputActionData())
+			{
+				const FCommonInputTypeInfo InputTypeInfo = DataTableRow->GetCurrentInputTypeInfo(GetInputSubsystem());
+				HoldTime = InputTypeInfo.HoldTime;
+				HoldRollbackTime = InputTypeInfo.HoldRollbackTime;
+			}
+		}
+	}
+	else
+	{
+		Super::UpdateHoldData(CurrentInputType);
+	}
+}
+
 #undef LOCTEXT_NAMESPACE
diff --git a/Engine/Plugins/Runtime/CommonUI/Source/CommonUI/Public/Input/CommonBoundActionButton.h b/Engine/Plugins/Runtime/CommonUI/Source/CommonUI/Public/Input/CommonBoundActionButton.h
--- a/Engine/Plugins/Runtime/CommonUI/Source/CommonUI/Public/Input/CommonBoundActionButton.h
+++ b/Engine/Plugins/Runtime/CommonUI/Source/CommonUI/Public/Input/CommonBoundActionButton.h
@@ -26,6 +26,7 @@ protected:
 	UE_API virtual void NativeOnCurrentTextStyleChanged() override;
 
 	UE_API virtual void UpdateInputActionWidget() override;
+	UE_API virtual void UpdateHoldData(ECommonInputType CurrentInputType) override;
 
 	UFUNCTION(BlueprintImplementableEvent, Category = "Common Bound Action")
 	UE_API void OnUpdateInputAction();
@@ -33,6 +34,10 @@ protected:
 protected:
 	UPROPERTY(BlueprintReadOnly, meta = (BindWidget), Category = "Text Block")
 	TObjectPtr<UCommonTextBlock> Text_ActionName;
+	
+	/** Set to true if clicking this button should require holding it when the bound action is a hold action */
+	UPROPERTY(EditAnywhere, BlueprintReadOnly, Category = "Input|Hold")
+	bool bLinkRequiresHoldToBindingHold = false;
 
 private:
 	FUIActionBindingHandle BindingHandle;
