[CommonInput] Make it possible for ControllerData to be added at runtime

#RNX

diff --git a/Engine/Plugins/Runtime/CommonUI/Source/CommonInput/Private/CommonInputBaseTypes.cpp b/Engine/Plugins/Runtime/CommonUI/Source/CommonInput/Private/CommonInputBaseTypes.cpp
--- a/Engine/Plugins/Runtime/CommonUI/Source/CommonInput/Private/CommonInputBaseTypes.cpp
+++ b/Engine/Plugins/Runtime/CommonUI/Source/CommonInput/Private/CommonInputBaseTypes.cpp
@@ -248,6 +248,7 @@ void UCommonInputPlatformSettings::InitializeControllerData() const
 {
 	if (ControllerData.Num() != ControllerDataClasses.Num())
 	{
+		ControllerDataClasses.Reset();
 		for (TSoftClassPtr<UCommonInputBaseControllerData> ControllerDataPtr : ControllerData)
 		{
 			if (TSubclassOf<UCommonInputBaseControllerData> ControllerDataClass = ControllerDataPtr.LoadSynchronous())
@@ -333,6 +334,35 @@ TArray<const UCommonInputBaseControllerData*> UCommonInputPlatformSettings::GetC
 	return ControllerDataForInputType;
 }
 
+void UCommonInputPlatformSettings::AddControllerDataEntry(TSoftClassPtr<UCommonInputBaseControllerData> Entry)
+{
+	if (ensure(!Entry.IsNull()))
+	{
+		ControllerData.AddUnique(Entry);
+		if (ControllerData.Num() == (ControllerDataClasses.Num() + 1))
+		{
+			if (TSubclassOf<UCommonInputBaseControllerData> ControllerDataClass = Entry.LoadSynchronous())
+			{
+				ControllerDataClasses.Add(ControllerDataClass);
+			}
+		}
+		else
+		{
+			InitializeControllerData();
+		}
+	}
+}
+
+void UCommonInputPlatformSettings::RemoveControllerDataEntry(TSoftClassPtr<UCommonInputBaseControllerData> Entry)
+{
+	ControllerData.RemoveAllSwap([&Entry](const TSoftClassPtr<UCommonInputBaseControllerData>& ControllerDataEntry) {
+		return ControllerDataEntry.GetUniqueID() == Entry.GetUniqueID();
+	});
+	ControllerDataClasses.RemoveAllSwap([&Entry](const TSubclassOf<UCommonInputBaseControllerData>& ControllerDataClassEntry){
+		return ControllerDataClassEntry.Get() == Entry.GetUniqueID();
+	});
+}
+
 FName UCommonInputPlatformSettings::GetBestGamepadNameForHardware(FName CurrentGamepadName, FName InputDeviceName, const FString& HardwareDeviceIdentifier)
 {
 	InitializeControllerData();
diff --git a/Engine/Plugins/Runtime/CommonUI/Source/CommonInput/Public/CommonInputBaseTypes.h b/Engine/Plugins/Runtime/CommonUI/Source/CommonInput/Public/CommonInputBaseTypes.h
--- a/Engine/Plugins/Runtime/CommonUI/Source/CommonInput/Public/CommonInputBaseTypes.h
+++ b/Engine/Plugins/Runtime/CommonUI/Source/CommonInput/Public/CommonInputBaseTypes.h
@@ -232,7 +232,13 @@ public:
 	UE_API bool TryGetInputBrush(FSlateBrush& OutBrush, FKey Key, ECommonInputType InputType, const FName GamepadName) const;
 	UE_API bool TryGetInputBrush(FSlateBrush& OutBrush, const TArray<FKey>& Keys, ECommonInputType InputType, const FName GamepadName) const;
 
+	TArray<TSoftClassPtr<UCommonInputBaseControllerData>> GetControllerData() const
+	{
+		return ControllerData;
+	}
 	UE_API TArray<const UCommonInputBaseControllerData*> GetControllerDataForInputType(ECommonInputType InputType, const FName GamepadName) const;
+	UE_API void AddControllerDataEntry(TSoftClassPtr<UCommonInputBaseControllerData> Entry);
+	UE_API void RemoveControllerDataEntry(TSoftClassPtr<UCommonInputBaseControllerData> Entry);
 
 	UE_API FName GetBestGamepadNameForHardware(FName CurrentGamepadName, FName InputDeviceName, const FString& HardwareDeviceIdentifier); 
 
@@ -253,12 +259,6 @@ public:
 		return bCanChangeGamepadType;
 	}
 
-	TArray<TSoftClassPtr<UCommonInputBaseControllerData>> GetControllerData()
-	{
-		return ControllerData;
-	}
-
-
 #if WITH_EDITOR
 	UE_API void PostEditChangeProperty(struct FPropertyChangedEvent& PropertyChangedEvent);
 #endif
