CommonUI: Add EditCondition = "bUseActionDomainDesiredInputConfig" for action domain input config properties

#rnx

diff --git a/Engine/Plugins/Runtime/CommonUI/Source/CommonInput/Public/CommonInputActionDomain.h b/Engine/Plugins/Runtime/CommonUI/Source/CommonInput/Public/CommonInputActionDomain.h
--- a/Engine/Plugins/Runtime/CommonUI/Source/CommonInput/Public/CommonInputActionDomain.h
+++ b/Engine/Plugins/Runtime/CommonUI/Source/CommonInput/Public/CommonInputActionDomain.h
@@ -11,6 +11,8 @@
 
 #define UE_API COMMONINPUT_API
 
+class UCommonInputActionDomain;
+
 DECLARE_LOG_CATEGORY_EXTERN(LogUIActionDomain, Log, All);
 
 UENUM()
@@ -53,13 +55,13 @@ public:
 	UPROPERTY(EditDefaultsOnly, Category = "Default")
 	ECommonInputEventFlowBehavior InnerBehavior = ECommonInputEventFlowBehavior::BlockIfHandled;
 
-	UPROPERTY(EditDefaultsOnly, Category = "Default" )
+	UPROPERTY(EditDefaultsOnly, Category = "Default")
 	bool bUseActionDomainDesiredInputConfig;
 
-	UPROPERTY(EditDefaultsOnly, Category = "Default")
+	UPROPERTY(EditDefaultsOnly, Category = "Default", meta = (EditCondition = "bUseActionDomainDesiredInputConfig"))
 	ECommonInputMode InputMode = ECommonInputMode::Game;
 
-	UPROPERTY(EditDefaultsOnly, Category = "Default")
+	UPROPERTY(EditDefaultsOnly, Category = "Default", meta = (EditCondition = "bUseActionDomainDesiredInputConfig"))
 	EMouseCaptureMode MouseCaptureMode = EMouseCaptureMode::CapturePermanently;
 
 	UE_API bool ShouldBreakInnerEventFlow(bool bInputEventHandled) const;
