PR #13080: Enhance debugging support of Action Domains for the CommonUI.DumpActivatableTree command

diff --git a/Engine/Plugins/Runtime/CommonUI/Source/CommonUI/Private/Input/CommonUIActionRouterBase.cpp b/Engine/Plugins/Runtime/CommonUI/Source/CommonUI/Private/Input/CommonUIActionRouterBase.cpp
--- a/Engine/Plugins/Runtime/CommonUI/Source/CommonUI/Private/Input/CommonUIActionRouterBase.cpp
+++ b/Engine/Plugins/Runtime/CommonUI/Source/CommonUI/Private/Input/CommonUIActionRouterBase.cpp
@@ -1106,6 +1106,28 @@ FGameplayTagContainer UCommonUIActionRouterBase::GetGameplayTagsForInputMode(con
 	return Tags;
 }
 
+void UCommonUIActionRouterBase::DebugDumpActionDomainRootNodes(int32 UserIndex, int32 ControllerId, bool bIncludeActions, bool bIncludeChildren, bool bIncludeInactive) const
+{
+	FString ActionDomainsOutputStr;
+	ActionDomainsOutputStr.Append(TEXT("******** Start Debugging ActionDomainRootNodes ********\n"));
+	for (const TPair<TObjectPtr<UCommonInputActionDomain>, FActionDomainSortedRootList>& Pair : ActionDomainRootNodes)
+	{
+		ActionDomainsOutputStr.Append(TEXT("\n****** Dumping ActionDomainRootNodes for ActionDomain: "));
+		ActionDomainsOutputStr.Append(GetNameSafe(Pair.Key) + TEXT(" ******\n"));
+		if (Pair.Value.GetRootList().Num() > 0)
+		{
+			Pair.Value.DebugDumpRootList(ActionDomainsOutputStr, bIncludeActions, bIncludeChildren, bIncludeInactive);
+			ActionDomainsOutputStr.Append(TEXT("\n\n"));
+		}
+		else
+		{
+			ActionDomainsOutputStr.Append(TEXT("-No root nodes found\n"));
+		}
+	}
+	ActionDomainsOutputStr.Append(TEXT("\n******** End Debugging ActionDomainRootNodes ********\n"));
+	UE_LOG(LogUIActionRouter, Display, TEXT("Dumping ActionDomainRootNodes for LocalPlayer [User %d, ControllerId %d]:\n%s\n"), UserIndex, ControllerId, *ActionDomainsOutputStr);
+}
+
 void UCommonUIActionRouterBase::ProcessRebuiltWidgets()
 {
 	// Begin by organizing all of the widgets that need nodes according to their direct parent
@@ -1959,7 +1981,9 @@ public:
 					ActionRouter->PersistentActions->DumpActionBindings(TreeOutputStr);
 				}
 
-				UE_LOG(LogUIActionRouter, Display, TEXT("Dumping ActivatableWidgetTree for LocalPlayer [User %d, ControllerId %d]:\n\n%s\n"), CurrIdx, LocalPlayer->GetControllerId(), *TreeOutputStr);
+				UE_LOG(LogUIActionRouter, Display, TEXT("Dumping ActivatableWidgetTree for LocalPlayer [User %d, ControllerId %d]:\n\n%s\n\n"), CurrIdx, LocalPlayer->GetControllerId(), *TreeOutputStr);
+				
+				ActionRouter->DebugDumpActionDomainRootNodes(CurrIdx, LocalPlayer->GetControllerId(), bIncludeActions, bIncludeChildren, bIncludeInactive);
 			}
 		}
 	}
@@ -2055,3 +2079,11 @@ void UCommonUIActionRouterBase::FActionDomainSortedRootList::Sort()
 	RootList.Sort(SortFunc);
 }
 
+void UCommonUIActionRouterBase::FActionDomainSortedRootList::DebugDumpRootList(FString& OutputStr, bool bIncludeActions, bool bIncludeChildren, bool bIncludeInactive) const
+{
+	for (const FActivatableTreeRootRef& Root : RootList)
+	{
+		Root->DebugDump(OutputStr, bIncludeActions, bIncludeChildren, bIncludeInactive);
+	}
+}
+
diff --git a/Engine/Plugins/Runtime/CommonUI/Source/CommonUI/Public/Input/CommonUIActionRouterBase.h b/Engine/Plugins/Runtime/CommonUI/Source/CommonUI/Public/Input/CommonUIActionRouterBase.h
--- a/Engine/Plugins/Runtime/CommonUI/Source/CommonUI/Public/Input/CommonUIActionRouterBase.h
+++ b/Engine/Plugins/Runtime/CommonUI/Source/CommonUI/Public/Input/CommonUIActionRouterBase.h
@@ -201,6 +201,8 @@ private:
 	UE_API bool ProcessInputOnActionDomains(ECommonInputMode ActiveInputMode, FKey Key, EInputEvent InputEvent, int32 UserIndex) const;
 	UE_API EProcessHoldActionResult ProcessHoldInputOnActionDomains(ECommonInputMode ActiveInputMode, FKey Key, EInputEvent InputEvent, int32 UserIndex) const;
 	UE_API FGameplayTagContainer GetGameplayTagsForInputMode(const ECommonInputMode Mode) const;
+	
+	void DebugDumpActionDomainRootNodes(int32 UserIndex, int32 ControllerId, bool bIncludeActions, bool bIncludeChildren, bool bIncludeInactive) const;
 
 	struct FPendingWidgetRegistration
 	{
@@ -280,6 +282,9 @@ PRAGMA_ENABLE_DEPRECATION_WARNINGS
 		bool Contains(FActivatableTreeRootRef RootNode) const;
 		
 		void Sort();
+		
+		// Needed for debug purposes
+		void DebugDumpRootList(FString& OutputStr, bool bIncludeActions, bool bIncludeChildren, bool bIncludeInactive) const;
 
 	private:
 		TArray<FActivatableTreeRootRef> RootList;
