Fix action router not respecting enabled state of activatable tree when using action domains.

#rnx
#Tests Ran game in a single client. Added splitscreen player and confirmed additional bindings were not being populated.

diff --git a/Engine/Plugins/Runtime/CommonUI/Source/CommonUI/Private/Input/CommonUIActionRouterBase.cpp b/Engine/Plugins/Runtime/CommonUI/Source/CommonUI/Private/Input/CommonUIActionRouterBase.cpp
--- a/Engine/Plugins/Runtime/CommonUI/Source/CommonUI/Private/Input/CommonUIActionRouterBase.cpp
+++ b/Engine/Plugins/Runtime/CommonUI/Source/CommonUI/Private/Input/CommonUIActionRouterBase.cpp
@@ -447,6 +447,13 @@ TArray<const UWidget*> UCommonUIActionRouterBase::GatherActiveAnalogScrollRecipi
 TArray<FUIActionBindingHandle> UCommonUIActionRouterBase::GatherActiveBindings() const
 {
 	TArray<FUIActionBindingHandle> BindingHandles = PersistentActions->GetActionBindings();
+
+	if (!bIsActivatableTreeEnabled)
+	{
+		// If we are ignoring the activatable tree, all active roots should be ignored.
+		return BindingHandles;
+	}
+
 	if (ActiveRootNode)
 	{
 		ActiveRootNode->AppendAllActiveActions(BindingHandles);
@@ -532,9 +539,9 @@ ERouteUIInputResult UCommonUIActionRouterBase::ProcessInput(FKey Key, EInputEven
 	{
 		EProcessHoldActionResult ProcessHoldResult = ActionRouter.PersistentActions->ProcessHoldInput(ActiveMode, Key, InputEvent, OwningUserIndex);
 	
-		if (ProcessHoldResult == EProcessHoldActionResult::Unhandled)
+		if (ProcessHoldResult == EProcessHoldActionResult::Unhandled && ActionRouter.bIsActivatableTreeEnabled)
 		{
-			if (ActionRouter.bIsActivatableTreeEnabled && ActionRouter.ActiveRootNode)
+			if (ActionRouter.ActiveRootNode)
 			{
 				ProcessHoldResult = ActionRouter.ActiveRootNode->ProcessHoldInput(ActiveMode, Key, InputEvent, OwningUserIndex);
 			}
@@ -552,9 +559,9 @@ ERouteUIInputResult UCommonUIActionRouterBase::ProcessInput(FKey Key, EInputEven
 	{
 		bool bHandled = ActionRouter.PersistentActions->ProcessNormalInput(ActiveMode, Key, Event, OwningUserIndex);
 
-		if (!bHandled)
+		if (!bHandled && ActionRouter.bIsActivatableTreeEnabled)
 		{
-			if (ActionRouter.bIsActivatableTreeEnabled && ActionRouter.ActiveRootNode)
+			if (ActionRouter.ActiveRootNode)
 			{
 				bHandled = ActionRouter.ActiveRootNode->ProcessNormalInput(ActiveMode, Key, Event, OwningUserIndex);
 			}
