[CommonUI] [Fix] Avoid duplicating RootNodes when calling RebuildWidget

- Fixes an issue where CommonUI would create two root nodes for the same widget if the the user calls the RebuildWidget method when the widget is getting constructed or activated.

#tests In Editor
#rnx

diff --git a/Engine/Plugins/Runtime/CommonUI/Source/CommonUI/Private/Input/CommonUIActionRouterBase.cpp b/Engine/Plugins/Runtime/CommonUI/Source/CommonUI/Private/Input/CommonUIActionRouterBase.cpp
--- a/Engine/Plugins/Runtime/CommonUI/Source/CommonUI/Private/Input/CommonUIActionRouterBase.cpp
+++ b/Engine/Plugins/Runtime/CommonUI/Source/CommonUI/Private/Input/CommonUIActionRouterBase.cpp
@@ -1931,7 +1931,7 @@ void UCommonUIActionRouterBase::HandleActivatableWidgetRebuilding(UCommonActivat
 {
 	if (RebuildingWidget.GetOwningLocalPlayer() == GetLocalPlayerChecked())
 	{
-		RebuiltWidgetsPendingNodeAssignment.Add(&RebuildingWidget);
+		RebuiltWidgetsPendingNodeAssignment.AddUnique(&RebuildingWidget);
 	}
 }
 
