Remove hold to purchase from RMT's purchases on mobile (when a secondary modal appears after hitting purchase)

#tests PIE
#rnx

diff --git a/Engine/Plugins/Runtime/CommonUI/Source/CommonUI/Private/CommonButtonBase.cpp b/Engine/Plugins/Runtime/CommonUI/Source/CommonUI/Private/CommonButtonBase.cpp
--- a/Engine/Plugins/Runtime/CommonUI/Source/CommonUI/Private/CommonButtonBase.cpp
+++ b/Engine/Plugins/Runtime/CommonUI/Source/CommonUI/Private/CommonButtonBase.cpp
@@ -2137,11 +2137,17 @@ void UCommonButtonBase::DisableButton()
 
 void UCommonButtonBase::SetRequiresHold(bool bInRequiresHold)
 {
+	const bool bPrevRequiresHold = bRequiresHold;
 	bRequiresHold = bInRequiresHold;
 	if (const UCommonInputSubsystem* CommonInputSubsystem = GetInputSubsystem())
 	{
 		UpdateHoldData(CommonInputSubsystem->GetCurrentInputType());
 	}
+
+	if (bPrevRequiresHold != bRequiresHold)
+	{
+		BP_OnRequiresHoldChanged();
+	}
 }
 
 void UCommonButtonBase::SetIsFocusable(bool bInIsFocusable)
diff --git a/Engine/Plugins/Runtime/CommonUI/Source/CommonUI/Public/CommonButtonBase.h b/Engine/Plugins/Runtime/CommonUI/Source/CommonUI/Public/CommonButtonBase.h
--- a/Engine/Plugins/Runtime/CommonUI/Source/CommonUI/Public/CommonButtonBase.h
+++ b/Engine/Plugins/Runtime/CommonUI/Source/CommonUI/Public/CommonButtonBase.h
@@ -703,6 +703,9 @@ protected:
 	COMMONUI_API void OnCurrentTextStyleChanged();
 	COMMONUI_API virtual void NativeOnCurrentTextStyleChanged();
 
+	UFUNCTION(BlueprintImplementableEvent, Category = CommonButton, meta = (DisplayName = "On Requires Hold Changed"))
+	COMMONUI_API void BP_OnRequiresHoldChanged();
+
 	/** Internal method to allow the selected state to be set regardless of selectability or toggleability */
 	UFUNCTION(BlueprintCallable, meta=(BlueprintProtected="true"), Category = "Common Button")
 	COMMONUI_API void SetSelectedInternal(bool bInSelected, bool bAllowSound = true, bool bBroadcast = true);
