CommonUI: Updated CommonWidgetCarousel to be more consistent in it's widget caching. Added an option to determine whether or not to cache widgets at all

Old behavior was to cache widgets in the carousel's RebuildWidget, but widgets added at runtime via AddChild would never be cached. Currently caching is enabled by default to match the behavior when children were added at design-time, though this could result in increased memory usage in use cases where all widgets were added dynamically.

diff --git a/Engine/Plugins/Runtime/CommonUI/Source/CommonUI/Private/CommonWidgetCarousel.cpp b/Engine/Plugins/Runtime/CommonUI/Source/CommonUI/Private/CommonWidgetCarousel.cpp
--- a/Engine/Plugins/Runtime/CommonUI/Source/CommonUI/Private/CommonWidgetCarousel.cpp
+++ b/Engine/Plugins/Runtime/CommonUI/Source/CommonUI/Private/CommonWidgetCarousel.cpp
@@ -12,6 +12,7 @@ UCommonWidgetCarousel::UCommonWidgetCarousel(const FObjectInitializer& ObjectIni
 	SetClipping(EWidgetClipping::ClipToBounds);
 
 	MoveSpeed = 5.f;
+	bCacheChildren = true;
 }
 
 void UCommonWidgetCarousel::ReleaseSlateResources(bool bReleaseChildren)
@@ -89,6 +90,33 @@ float UCommonWidgetCarousel::GetMoveSpeed() const
 	return MoveSpeed;
 }
 
+void UCommonWidgetCarousel::SetCacheChildren(bool InCacheChildren)
+{
+	if (bCacheChildren != InCacheChildren)
+	{
+		bCacheChildren = InCacheChildren;
+		if (bCacheChildren)
+		{
+			for (UPanelSlot* PanelSlot : Slots)
+			{
+				if (PanelSlot->Content)
+				{
+					CachedSlotWidgets.AddUnique(PanelSlot->Content->TakeWidget());
+				}
+			}
+		}
+		else
+		{
+			CachedSlotWidgets.Empty();
+		}
+	}
+}
+
+bool UCommonWidgetCarousel::GetCacheChildren() const
+{
+	return bCacheChildren;
+}
+
 int32 UCommonWidgetCarousel::GetActiveWidgetIndex() const
 {
 	if ( MyCommonWidgetCarousel.IsValid() )
@@ -138,12 +166,33 @@ UClass* UCommonWidgetCarousel::GetSlotClass() const
 
 void UCommonWidgetCarousel::OnSlotAdded(UPanelSlot* InSlot)
 {
+	if (bCacheChildren && InSlot && InSlot->Content)
+	{
+		const TSharedPtr<SWidget> SafeWidget = InSlot->Content->TakeWidget();
+		if (SafeWidget.IsValid())
+		{
+			CachedSlotWidgets.AddUnique(SafeWidget.ToSharedRef());
+		}
+	}
+
 	if (MyCommonWidgetCarousel)
 	{
 		MyCommonWidgetCarousel->GenerateCurrentWidgets();
 	}
 }
 
+void UCommonWidgetCarousel::OnSlotRemoved(UPanelSlot* InSlot)
+{
+	if (bCacheChildren && InSlot && InSlot->Content)
+	{
+		const TSharedPtr<SWidget> SafeWidget = InSlot->Content->GetCachedWidget();
+		if (SafeWidget.IsValid())
+		{
+			CachedSlotWidgets.Remove(SafeWidget.ToSharedRef());
+		}
+	}
+}
+
 TSharedRef<SWidget> UCommonWidgetCarousel::RebuildWidget()
 {
 	MyCommonWidgetCarousel = SNew(SWidgetCarousel<UPanelSlot*>)
@@ -155,10 +204,12 @@ TSharedRef<SWidget> UCommonWidgetCarousel::RebuildWidget()
 		.OnGenerateWidget_UObject(this, &UCommonWidgetCarousel::OnGenerateWidgetForCarousel)
 		.OnPageChanged_UObject(this, &UCommonWidgetCarousel::HandlePageChanged);
 
+	CachedSlotWidgets.Empty();
+
 	for (UPanelSlot* PanelSlot : Slots)
 	{
 		PanelSlot->Parent = this;
-		if (PanelSlot->Content)
+		if (bCacheChildren && PanelSlot->Content)
 		{
 			CachedSlotWidgets.AddUnique(PanelSlot->Content->TakeWidget());
 		}
diff --git a/Engine/Plugins/Runtime/CommonUI/Source/CommonUI/Public/CommonWidgetCarousel.h b/Engine/Plugins/Runtime/CommonUI/Source/CommonUI/Public/CommonWidgetCarousel.h
--- a/Engine/Plugins/Runtime/CommonUI/Source/CommonUI/Public/CommonWidgetCarousel.h
+++ b/Engine/Plugins/Runtime/CommonUI/Source/CommonUI/Public/CommonWidgetCarousel.h
@@ -34,6 +34,10 @@ protected:
 	/** How quickly the carousel transitions when changing active widget */
 	UPROPERTY(EditAnywhere, BlueprintReadWrite, Getter, Setter, BlueprintGetter="GetMoveSpeed", BlueprintSetter="SetMoveSpeed", Category="Carousel", meta=( UIMin=0, ClampMin=0 ))
 	float MoveSpeed;
+
+	/** Whether we should cache children to prevent them from being destroyed when not visible. Enable this to avoid ConstructWidget costs when changing index, disable to save memory  */
+	UPROPERTY(EditAnywhere, BlueprintReadWrite, Getter = "GetCacheChildren", Setter = "SetCacheChildren", BlueprintGetter = "GetCacheChildren", BlueprintSetter = "SetCacheChildren", Category = "Carousel")
+	bool bCacheChildren;
 	
 public:
 
@@ -72,6 +76,14 @@ public:
 	/** Gets the Move Speed. */
 	UFUNCTION(BlueprintCallable, Category = "Widget")
 	UE_API float GetMoveSpeed() const;
+
+	/** Sets the current caching behavior. */
+	UFUNCTION(BlueprintCallable, Category="Widget")
+	UE_API void SetCacheChildren(bool InCacheChildren);
+
+	/** Gets the current caching behavior. */
+	UFUNCTION(BlueprintCallable, Category = "Widget")
+	UE_API bool GetCacheChildren() const;
 	
 	UPROPERTY(BlueprintAssignable, Category = "Carousel")
 	FOnCurrentPageIndexChanged OnCurrentPageIndexChanged;
@@ -100,6 +112,7 @@ protected:
 	// UPanelWidget
 	UE_API virtual UClass* GetSlotClass() const override;
 	UE_API virtual void OnSlotAdded(UPanelSlot* InSlot) override;
+	UE_API virtual void OnSlotRemoved(UPanelSlot* InSlot) override;
 	// End UPanelWidget
 
 	// UWidget interface
