Common UI:
- Add property in Action Domain Table to set if the mouse cursor should be hidden or not when applying the default UI Input Config

#rnx

#changelist validated
#virtualized

diff --git a/Engine/Plugins/Runtime/CommonUI/Source/CommonInput/Public/CommonInputActionDomain.h b/Engine/Plugins/Runtime/CommonUI/Source/CommonInput/Public/CommonInputActionDomain.h
--- a/Engine/Plugins/Runtime/CommonUI/Source/CommonInput/Public/CommonInputActionDomain.h
+++ b/Engine/Plugins/Runtime/CommonUI/Source/CommonInput/Public/CommonInputActionDomain.h
@@ -87,6 +87,9 @@ public:
 
 	UPROPERTY(EditDefaultsOnly, Category = "Default")
 	EMouseCaptureMode MouseCaptureMode = EMouseCaptureMode::CapturePermanently;
+
+	UPROPERTY(EditDefaultsOnly, Category = "Default")
+	bool bHideCursorDuringViewportCapture = true;
 };
 
 #undef UE_API
diff --git a/Engine/Plugins/Runtime/CommonUI/Source/CommonUI/Private/Input/CommonUIActionRouterBase.cpp b/Engine/Plugins/Runtime/CommonUI/Source/CommonUI/Private/Input/CommonUIActionRouterBase.cpp
--- a/Engine/Plugins/Runtime/CommonUI/Source/CommonUI/Private/Input/CommonUIActionRouterBase.cpp
+++ b/Engine/Plugins/Runtime/CommonUI/Source/CommonUI/Private/Input/CommonUIActionRouterBase.cpp
@@ -1745,7 +1745,7 @@ void UCommonUIActionRouterBase::RefreshActionDomainLeafNodeConfig()
 		}
 		else
 		{
-			SetActiveUIInputConfig(FUIInputConfig(ActionDomainTable->InputMode, ActionDomainTable->MouseCaptureMode), ActionDomainTable);
+			SetActiveUIInputConfig(FUIInputConfig(ActionDomainTable->InputMode, ActionDomainTable->MouseCaptureMode, ActionDomainTable->bHideCursorDuringViewportCapture), ActionDomainTable);
 		}
 	}
 }
