Fixed issue where SCommonButton wouldn't release when it loses focus

diff --git a/Engine/Plugins/Runtime/CommonUI/Source/CommonUI/Private/CommonButtonTypes.cpp b/Engine/Plugins/Runtime/CommonUI/Source/CommonUI/Private/CommonButtonTypes.cpp
--- a/Engine/Plugins/Runtime/CommonUI/Source/CommonUI/Private/CommonButtonTypes.cpp
+++ b/Engine/Plugins/Runtime/CommonUI/Source/CommonUI/Private/CommonButtonTypes.cpp
@@ -205,6 +205,8 @@ FReply SCommonButton::OnFocusReceived(const FGeometry& MyGeometry, const FFocusE
 
 void SCommonButton::OnFocusLost(const FFocusEvent& InFocusEvent)
 {
+	SButton::OnFocusLost(InFocusEvent);
+
 	OnLostFocus.ExecuteIfBound();
 }
 
