Cached a few hot fname queries.
Saves around 0.1ms on GT, 0.15ms on RT on Switch ReplayRun with FName compression enabled.

diff --git a/Engine/Plugins/Runtime/CommonUI/Source/CommonUI/Private/Input/UIActionRouterTypes.cpp b/Engine/Plugins/Runtime/CommonUI/Source/CommonUI/Private/Input/UIActionRouterTypes.cpp
--- a/Engine/Plugins/Runtime/CommonUI/Source/CommonUI/Private/Input/UIActionRouterTypes.cpp
+++ b/Engine/Plugins/Runtime/CommonUI/Source/CommonUI/Private/Input/UIActionRouterTypes.cpp
@@ -228,14 +228,15 @@ FUIActionBindingHandle FUIActionBinding::TryCreate(const UWidget& InBoundWidget,
 FUIActionBindingHandle FUIActionBinding::TryCreate(const UWidget& InBoundWidget, const FBindUIActionArgs& BindArgs, int32 UserIndex)
 {
 	bool bIsEnhancedInputSupportEnabled = CommonUI::IsEnhancedInputSupportEnabled();
-	if (BindArgs.GetActionName().IsNone() && (!bIsEnhancedInputSupportEnabled || !BindArgs.InputAction.IsValid()))
+	const FName BindActionName = BindArgs.GetActionName();
+	if (BindActionName.IsNone() && (!bIsEnhancedInputSupportEnabled || !BindArgs.InputAction.IsValid()))
 	{
 		UE_LOG(LogUIActionRouter, Error, TEXT("Cannot create action binding for widget [%s] - no action provided."), *InBoundWidget.GetName());
 		return FUIActionBindingHandle();
 	}
 	else if (!BindArgs.OnExecuteAction.IsBound())
 	{
-		UE_LOG(LogUIActionRouter, Error, TEXT("Cannot bind widget [%s] to action [%s] - empty handler delegate."), *InBoundWidget.GetName(), *BindArgs.GetActionName().ToString());
+		UE_LOG(LogUIActionRouter, Error, TEXT("Cannot bind widget [%s] to action [%s] - empty handler delegate."), *InBoundWidget.GetName(), *BindActionName.ToString());
 		return FUIActionBindingHandle();
 	}
 	else if (BindArgs.ActionTag.IsValid() && !UCommonUIInputSettings::Get().FindAction(BindArgs.ActionTag))
@@ -245,24 +246,24 @@ FUIActionBindingHandle FUIActionBinding::TryCreate(const UWidget& InBoundWidget,
 	}
 	else if (!BindArgs.LegacyActionTableRow.IsNull() && !BindArgs.LegacyActionTableRow.GetRow<FCommonInputActionDataBase>(TEXT("")))
 	{
-		UE_LOG(LogUIActionRouter, Error, TEXT("Cannot bind widget [%s] to action [%s] - provided legacy data table row does not resolve to valid data."), *InBoundWidget.GetName(), *BindArgs.GetActionName().ToString());
+		UE_LOG(LogUIActionRouter, Error, TEXT("Cannot bind widget [%s] to action [%s] - provided legacy data table row does not resolve to valid data."), *InBoundWidget.GetName(), *BindActionName.ToString());
 		return FUIActionBindingHandle();
 	}
 	else if (!BindArgs.ActionTag.IsValid() && BindArgs.LegacyActionTableRow.IsNull() && bIsEnhancedInputSupportEnabled && !BindArgs.InputAction.IsValid())
 	{
-		UE_LOG(LogUIActionRouter, Error, TEXT("Cannot bind widget [%s] to action [%s] - provided input action is invalid."), *InBoundWidget.GetName(), *BindArgs.GetActionName().ToString());
+		UE_LOG(LogUIActionRouter, Error, TEXT("Cannot bind widget [%s] to action [%s] - provided input action is invalid."), *InBoundWidget.GetName(), *BindActionName.ToString());
 		return FUIActionBindingHandle();
 	}
 	else if (UserIndex == INDEX_NONE)
 	{
-		UE_LOG(LogUIActionRouter, Error, TEXT("Cannot bind widget [%s] to action [%s] [%s] - invalid UserIndex."), *BindArgs.GetActionName().ToString(), InputEventToString(BindArgs.KeyEvent), *InBoundWidget.GetName());
+		UE_LOG(LogUIActionRouter, Error, TEXT("Cannot bind widget [%s] to action [%s] [%s] - invalid UserIndex."), *BindActionName.ToString(), InputEventToString(BindArgs.KeyEvent), *InBoundWidget.GetName());
 		return FUIActionBindingHandle();
 	}
 	
 	// Make sure there is no existing binding for the same action associated with the same widget
 	for (const TPair<FUIActionBindingHandle, TSharedPtr<FUIActionBinding>>& HandleBindingPair : AllRegistrationsByHandle)
 	{
-		if (HandleBindingPair.Value->ActionName == BindArgs.GetActionName() &&
+		if (HandleBindingPair.Value->ActionName == BindActionName &&
 			HandleBindingPair.Value->InputEvent == BindArgs.KeyEvent &&
 			HandleBindingPair.Value->BoundWidget.Get() == &InBoundWidget &&
 			HandleBindingPair.Value->UserIndex == UserIndex)
diff --git a/Engine/Source/Runtime/Core/Private/Serialization/Archive.cpp b/Engine/Source/Runtime/Core/Private/Serialization/Archive.cpp
--- a/Engine/Source/Runtime/Core/Private/Serialization/Archive.cpp
+++ b/Engine/Source/Runtime/Core/Private/Serialization/Archive.cpp
@@ -1465,7 +1465,8 @@ void FArchiveState::SetEngineVer(const FEngineVersionBase& InVer)
 
 void FArchiveState::SetEngineNetVer(const uint32 InEngineNetVer)
 {
-	SetCustomVersion(FEngineNetworkCustomVersion::Guid, InEngineNetVer, TEXT("EngineNetworkVersion"));
+	static const FName EngineNetworkVersionName("EngineNetworkVersion");
+	SetCustomVersion(FEngineNetworkCustomVersion::Guid, InEngineNetVer, EngineNetworkVersionName);
 }
 
 uint32 FArchiveState::EngineNetVer() const
@@ -1475,7 +1476,8 @@ uint32 FArchiveState::EngineNetVer() const
 
 void FArchiveState::SetGameNetVer(const uint32 InGameNetVer)
 {
-	SetCustomVersion(FGameNetworkCustomVersion::Guid, InGameNetVer, TEXT("GameNetworkVersion"));
+	static const FName GameNetworkVersionName("GameNetworkVersion");
+	SetCustomVersion(FGameNetworkCustomVersion::Guid, InGameNetVer, GameNetworkVersionName);
 }
 
 uint32 FArchiveState::GameNetVer() const
