In SCommonButtonTableRow we now listen to drag and drop events only when bAllowDragDrop is true. In SObjectTableRow we handle OnDrop only if bAllowDragDrop is true.

diff --git a/Engine/Plugins/Runtime/CommonUI/Source/CommonUI/Public/SCommonButtonTableRow.h b/Engine/Plugins/Runtime/CommonUI/Source/CommonUI/Public/SCommonButtonTableRow.h
--- a/Engine/Plugins/Runtime/CommonUI/Source/CommonUI/Public/SCommonButtonTableRow.h
+++ b/Engine/Plugins/Runtime/CommonUI/Source/CommonUI/Public/SCommonButtonTableRow.h
@@ -88,12 +88,15 @@ protected:
 			CommonButton->OnHovered().AddSP(this, &SCommonButtonTableRow::HandleButtonHovered);
 			CommonButton->OnUnhovered().AddSP(this, &SCommonButtonTableRow::HandleButtonUnhovered);
 			CommonButton->OnIsSelectedChanged().AddSP(this, &SCommonButtonTableRow::HandleButtonSelectionChanged);
-			
-			CommonButton->OnCommonButtonDragDetected().BindSP(this, &SCommonButtonTableRow::OnDragDetected);
-			CommonButton->OnCommonButtonDragEnter().BindSP(this, &SCommonButtonTableRow::OnDragEnter);
-			CommonButton->OnCommonButtonDragLeave().BindSP(this, &SCommonButtonTableRow::OnDragLeave);
-			CommonButton->OnCommonButtonDragOver().BindSP(this, &SCommonButtonTableRow::OnDragOver);
-			CommonButton->OnCommonButtonDrop().BindSP(this, &SCommonButtonTableRow::OnDrop);
+
+			if (this->GetAllowDragDrop())
+			{
+				CommonButton->OnCommonButtonDragDetected().BindSP(this, &SCommonButtonTableRow::OnDragDetected);
+				CommonButton->OnCommonButtonDragEnter().BindSP(this, &SCommonButtonTableRow::OnDragEnter);
+				CommonButton->OnCommonButtonDragLeave().BindSP(this, &SCommonButtonTableRow::OnDragLeave);
+				CommonButton->OnCommonButtonDragOver().BindSP(this, &SCommonButtonTableRow::OnDragOver);
+				CommonButton->OnCommonButtonDrop().BindSP(this, &SCommonButtonTableRow::OnDrop);
+			}
 
 			if (this->IsItemSelectable())
 			{
diff --git a/Engine/Source/Runtime/UMG/Public/Slate/SObjectTableRow.h b/Engine/Source/Runtime/UMG/Public/Slate/SObjectTableRow.h
--- a/Engine/Source/Runtime/UMG/Public/Slate/SObjectTableRow.h
+++ b/Engine/Source/Runtime/UMG/Public/Slate/SObjectTableRow.h
@@ -500,6 +500,11 @@ public:
 
 	virtual FReply OnDrop(const FGeometry& MyGeometry, const FDragDropEvent& DragDropEvent) override
 	{
+		if (!bAllowDragDrop)
+		{
+			return FReply::Unhandled();
+		}
+
 		FReply Reply = SObjectWidget::OnDrop(MyGeometry, DragDropEvent);
 
 		if (OnRowAcceptDrop.IsBound() && !Reply.IsEventHandled())
@@ -732,6 +737,11 @@ public:
 	}
 	// ~SWidget interface
 
+	bool GetAllowDragDrop() const
+	{
+		return bAllowDragDrop;
+	}
+
 protected:
 	virtual void InitializeObjectRow()
 	{
