CommonUI: Fix not being able to click buttons with Virtual_Accept while binding hold actions to Virtual_Accept

diff --git a/Engine/Plugins/Runtime/CommonUI/Source/CommonUI/Private/Input/CommonAnalogCursor.cpp b/Engine/Plugins/Runtime/CommonUI/Source/CommonUI/Private/Input/CommonAnalogCursor.cpp
--- a/Engine/Plugins/Runtime/CommonUI/Source/CommonUI/Private/Input/CommonAnalogCursor.cpp
+++ b/Engine/Plugins/Runtime/CommonUI/Source/CommonUI/Private/Input/CommonAnalogCursor.cpp
@@ -209,10 +209,10 @@ void FCommonAnalogCursor::Tick(const float DeltaTime, FSlateApplication& SlateAp
 										ScrollAmount,
 										FModifierKeysState());
 
-									UCommonInputSubsystem& InputSubsytem = ActionRouter.GetInputSubsystem();
-									InputSubsytem.SetIsGamepadSimulatedClick(true);
+									UCommonInputSubsystem& InputSubsystem = ActionRouter.GetInputSubsystem();
+									InputSubsystem.SetIsGamepadSimulatedClick(true);
 									SlateApp.ProcessMouseWheelOrGestureEvent(MouseEvent, nullptr);
-									InputSubsytem.SetIsGamepadSimulatedClick(false);
+									InputSubsystem.SetIsGamepadSimulatedClick(false);
 								}
 							}
 						}
@@ -264,10 +264,10 @@ bool FCommonAnalogCursor::HandleKeyDownEvent(FSlateApplication& SlateApp, const
 		else if (!bIsVirtualAccept || ShouldVirtualAcceptSimulateMouseButton(InKeyEvent, IE_Pressed))
 		{
 			// There is no awareness on a mouse event of whether it's real or not, so mark that here.
-			UCommonInputSubsystem& InputSubsytem = ActionRouter.GetInputSubsystem();
-			InputSubsytem.SetIsGamepadSimulatedClick(bIsVirtualAccept);
+			UCommonInputSubsystem& InputSubsystem = ActionRouter.GetInputSubsystem();
+			InputSubsystem.SetIsGamepadSimulatedClick(bIsVirtualAccept);
 			bool bReturnValue = FAnalogCursor::HandleKeyDownEvent(SlateApp, InKeyEvent);
-			InputSubsytem.SetIsGamepadSimulatedClick(false);
+			InputSubsystem.SetIsGamepadSimulatedClick(false);
 
 			return bReturnValue;
 		}
@@ -287,6 +287,8 @@ bool FCommonAnalogCursor::HandleKeyUpEvent(FSlateApplication& SlateApp, const FK
 		if (PressedKey == EKeys::Gamepad_RightTrigger) { ShoulderButtonStatus ^= EShoulderButtonFlags::RightTrigger; }
 #endif
 
+		TGuardValue<TOptional<FKeyEvent>> KeyUpEventGuard(ActiveKeyUpEvent, TOptional<FKeyEvent>(InKeyEvent));
+		
 		// We support binding actions to the virtual accept key, so it's a special flower that gets processed right now
 		const bool bIsVirtualAccept = InKeyEvent.GetKey() == EKeys::Virtual_Accept;
 		if (bIsVirtualAccept && ActionRouter.ProcessInput(InKeyEvent.GetKey(), IE_Released) == ERouteUIInputResult::Handled)
@@ -415,6 +417,17 @@ bool FCommonAnalogCursor::ShouldVirtualAcceptSimulateMouseButton(const FKeyEvent
 	return CVarShouldVirtualAcceptSimulateMouseButton.GetValueOnGameThread();
 }
 
+void FCommonAnalogCursor::OnVirtualAcceptHoldCanceled()
+{
+	if (ActiveKeyUpEvent.IsSet() && ActiveKeyUpEvent->GetKey() == EKeys::Virtual_Accept && ShouldVirtualAcceptSimulateMouseButton(ActiveKeyUpEvent.GetValue(), IE_Pressed))
+	{
+		UCommonInputSubsystem& InputSubsystem = ActionRouter.GetInputSubsystem();
+		InputSubsystem.SetIsGamepadSimulatedClick(true);
+		FAnalogCursor::HandleKeyDownEvent(FSlateApplication::Get(), ActiveKeyUpEvent.GetValue());
+		InputSubsystem.SetIsGamepadSimulatedClick(false);
+	}
+}
+
 //void FCommonAnalogCursor::SetCursorMovementStick(EAnalogStick InCursorMovementStick)
 //{
 //	const EAnalogStick NewStick = InCursorMovementStick == EAnalogStick::Max ? EAnalogStick::Left : InCursorMovementStick;
diff --git a/Engine/Plugins/Runtime/CommonUI/Source/CommonUI/Private/Input/CommonUIActionRouterBase.cpp b/Engine/Plugins/Runtime/CommonUI/Source/CommonUI/Private/Input/CommonUIActionRouterBase.cpp
--- a/Engine/Plugins/Runtime/CommonUI/Source/CommonUI/Private/Input/CommonUIActionRouterBase.cpp
+++ b/Engine/Plugins/Runtime/CommonUI/Source/CommonUI/Private/Input/CommonUIActionRouterBase.cpp
@@ -575,7 +575,7 @@ ERouteUIInputResult UCommonUIActionRouterBase::ProcessInput(FKey Key, EInputEven
 		return bHandled;
 	};
 
-	const auto ProcessInputOnActionRouter = [&ProcessHoldInputFunc, &ProcessNormalInputFunc, InputEvent](const UCommonUIActionRouterBase& ActionRouter)
+	const auto ProcessInputOnActionRouter = [&ProcessHoldInputFunc, &ProcessNormalInputFunc, &Key, InputEvent](const UCommonUIActionRouterBase& ActionRouter)
 	{
 		EProcessHoldActionResult ProcessHoldResult = ProcessHoldInputFunc(ActionRouter);
 		if (ProcessHoldResult == EProcessHoldActionResult::Handled)
@@ -586,7 +586,17 @@ ERouteUIInputResult UCommonUIActionRouterBase::ProcessInput(FKey Key, EInputEven
 		if (ProcessHoldResult == EProcessHoldActionResult::GeneratePress)
 		{
 			// A hold action was in progress but quickly aborted, so we want to generate a press action now for any normal bindings that are interested
-			ProcessNormalInputFunc(ActionRouter, IE_Pressed);
+			if (!ProcessNormalInputFunc(ActionRouter, IE_Pressed))
+			{
+				if (Key == EKeys::Virtual_Accept && ActionRouter.AnalogCursor.IsValid())
+				{
+					// If Virtual_Accept is bound to a Hold action, the Pressed event for it will not trigger a simulated click
+					// but when a hold action is canceled fast enough, we want to trigger any "normal" actions bound to the same key
+					// so we give the analog cursor a chance to trigger the simulated click here
+					// This allows users to bind Hold actions to Virtual_Accept while still triggering focused button widgets with quick virtual_accept presses
+					ActionRouter.AnalogCursor->OnVirtualAcceptHoldCanceled();
+				}
+			}
 		}
 
 		// Even if no widget cares about this input, we don't want to let anything through to the actual game while we're in menu mode
diff --git a/Engine/Plugins/Runtime/CommonUI/Source/CommonUI/Public/Input/CommonAnalogCursor.h b/Engine/Plugins/Runtime/CommonUI/Source/CommonUI/Public/Input/CommonAnalogCursor.h
--- a/Engine/Plugins/Runtime/CommonUI/Source/CommonUI/Public/Input/CommonAnalogCursor.h
+++ b/Engine/Plugins/Runtime/CommonUI/Source/CommonUI/Public/Input/CommonAnalogCursor.h
@@ -5,6 +5,7 @@
 #include "Engine/EngineBaseTypes.h"
 #include "Framework/Application/AnalogCursor.h"
 #include "InputCoreTypes.h"
+#include "Input/Events.h"
 #include "Rendering/SlateRenderTransform.h"
 
 #define UE_API COMMONUI_API
@@ -57,6 +58,8 @@ public:
 
 	UE_API virtual bool ShouldVirtualAcceptSimulateMouseButton(const FKeyEvent& InKeyEvent, EInputEvent InputEvent) const;
 
+	UE_API void OnVirtualAcceptHoldCanceled();
+
 protected:
 	UE_API FCommonAnalogCursor(const UCommonUIActionRouterBase& InActionRouter);
 	UE_API virtual void Initialize();
@@ -106,6 +109,8 @@ private:
 	TWeakPtr<SWidget> LastCursorTarget;
 	FSlateRenderTransform LastCursorTargetTransform;
 
+	TOptional<FKeyEvent> ActiveKeyUpEvent;
+
 	float TimeUntilScrollUpdate = 0.f;
 	
 #if !UE_BUILD_SHIPPING
