CommonUI: Include Action Binding UserIndex when dumping active roots

diff --git a/Engine/Plugins/Runtime/CommonUI/Source/CommonUI/Private/Input/UIActionRouterTypes.cpp b/Engine/Plugins/Runtime/CommonUI/Source/CommonUI/Private/Input/UIActionRouterTypes.cpp
--- a/Engine/Plugins/Runtime/CommonUI/Source/CommonUI/Private/Input/UIActionRouterTypes.cpp
+++ b/Engine/Plugins/Runtime/CommonUI/Source/CommonUI/Private/Input/UIActionRouterTypes.cpp
@@ -322,11 +322,12 @@ FCommonInputActionDataBase* FUIActionBinding::GetLegacyInputActionData() const
 
 FString FUIActionBinding::ToDebugString() const
 {
-	return FString::Printf(TEXT("%s: Owner [%s], Mode [%s], Displayed? [%s]"),
+	return FString::Printf(TEXT("%s: Owner [%s], Mode [%s], Displayed? [%s], UserIndex [%d]"),
 		*ActionName.ToString(),
 		BoundWidget.IsValid() ? *BoundWidget->GetName() : TEXT("Invalid"),
 		LexToString(InputMode),
-		*LexToString(bDisplayInActionBar));
+		*LexToString(bDisplayInActionBar),
+		UserIndex);
 }
 
 void FUIActionBinding::BeginHold()
