CommonUI: Fix persistent action bindings not always being registered as persistent actions

diff --git a/Engine/Plugins/Runtime/CommonUI/Source/CommonUI/Private/Input/CommonUIActionRouterBase.cpp b/Engine/Plugins/Runtime/CommonUI/Source/CommonUI/Private/Input/CommonUIActionRouterBase.cpp
--- a/Engine/Plugins/Runtime/CommonUI/Source/CommonUI/Private/Input/CommonUIActionRouterBase.cpp
+++ b/Engine/Plugins/Runtime/CommonUI/Source/CommonUI/Private/Input/CommonUIActionRouterBase.cpp
@@ -240,7 +240,17 @@ FUIActionBindingHandle UCommonUIActionRouterBase::RegisterUIActionBinding(const
 
 		if (OwnerNode)
 		{
-			OwnerNode->AddBinding(*FUIActionBinding::FindBinding(BindingHandle));
+			if (const TSharedPtr<FUIActionBinding> Binding = FUIActionBinding::FindBinding(BindingHandle))
+			{
+				if (Binding->bIsPersistent)
+				{
+					PersistentActions->AddBinding(*Binding);
+				}
+				else
+				{
+					OwnerNode->AddBinding(*Binding);
+				}
+			}
 			
 			if (UCommonActivatableWidget* ActivatableWidget = OwnerNode->GetWidget())
 			{
