Extract a new public method called "HadAnyChangeOfInputMethodInTheLastThrashingWindow" from CommonInputSubsystem to expose it to public API.

diff --git a/Engine/Plugins/Runtime/CommonUI/Source/CommonInput/Private/CommonInputSubsystem.cpp b/Engine/Plugins/Runtime/CommonUI/Source/CommonInput/Private/CommonInputSubsystem.cpp
--- a/Engine/Plugins/Runtime/CommonUI/Source/CommonInput/Private/CommonInputSubsystem.cpp
+++ b/Engine/Plugins/Runtime/CommonUI/Source/CommonInput/Private/CommonInputSubsystem.cpp
@@ -164,6 +164,13 @@ bool UCommonInputSubsystem::IsInputMethodActive(ECommonInputType InputMethod) co
 	return GetCurrentInputType() == InputMethod;
 }
 
+bool UCommonInputSubsystem::HadAnyChangeOfInputMethodInTheLastThrashingWindow() const
+{
+	const UCommonInputSettings& InputSettings = ICommonInputModule::GetSettings();
+	const double ChangeDelta = (FPlatformTime::Seconds() - LastInputMethodChangeTime);
+	return ChangeDelta < InputSettings.GetInputMethodThrashingWindowInSeconds();
+}
+
 TSharedPtr<FCommonInputPreprocessor> UCommonInputSubsystem::MakeInputProcessor()
 {
 	return MakeShared<FCommonInputPreprocessor>(*this);
@@ -227,8 +234,7 @@ bool UCommonInputSubsystem::CheckForInputMethodThrashing(ECommonInputType NewInp
 			return false;
 		}
 
-		const double ChangeDelta = (Now - LastInputMethodChangeTime);
-		if (ChangeDelta < InputSettings.GetInputMethodThrashingWindowInSeconds())
+		if (HadAnyChangeOfInputMethodInTheLastThrashingWindow())
 		{
 			NumberOfInputMethodChangesRecently++;
 			if (NumberOfInputMethodChangesRecently > InputSettings.GetInputMethodThrashingLimit())
diff --git a/Engine/Plugins/Runtime/CommonUI/Source/CommonInput/Public/CommonInputSubsystem.h b/Engine/Plugins/Runtime/CommonUI/Source/CommonInput/Public/CommonInputSubsystem.h
--- a/Engine/Plugins/Runtime/CommonUI/Source/CommonInput/Public/CommonInputSubsystem.h
+++ b/Engine/Plugins/Runtime/CommonUI/Source/CommonInput/Public/CommonInputSubsystem.h
@@ -101,6 +101,9 @@ public:
 	*/
 	static FPlatformInputSupportOverrideDelegate& GetOnPlatformInputSupportOverride() { return OnPlatformInputSupportOverride; }
 
+	/** Returns true if the input method changed in the last thrashing window. See UCommonInputSettings::GetInputMethodThrashingWindowInSeconds */
+	UE_API bool HadAnyChangeOfInputMethodInTheLastThrashingWindow() const;
+
 protected:
 	UE_API virtual TSharedPtr<FCommonInputPreprocessor> MakeInputProcessor();
 
