CommonButtonBase improvements

- Fixed the double call of the `UUserWidget::OnDrop` by returning `FReply::Handled` if it's handled by native handlers
- Removed the BIE's duplicates of the same BIEs we have in UUserWidget
- Fixed braodcasting of the wrong  delegate in `UCommonButtonBase::NativeOnDrop`

diff --git a/Engine/Plugins/Runtime/CommonUI/Source/CommonUI/Private/CommonButtonBase.cpp b/Engine/Plugins/Runtime/CommonUI/Source/CommonUI/Private/CommonButtonBase.cpp
--- a/Engine/Plugins/Runtime/CommonUI/Source/CommonUI/Private/CommonButtonBase.cpp
+++ b/Engine/Plugins/Runtime/CommonUI/Source/CommonUI/Private/CommonButtonBase.cpp
@@ -1634,7 +1634,10 @@ COMMONUI_API FReply UCommonButtonBase::HandleButtonDragOver(const FGeometry& MyG
 	TSharedPtr<FUMGDragDropOp> NativeOp = DragDropEvent.GetOperationAs<FUMGDragDropOp>();
 	if (NativeOp.IsValid())
 	{
-		NativeOnDragOver(MyGeometry, DragDropEvent, NativeOp->GetOperation());
+		if (NativeOnDragOver(MyGeometry, DragDropEvent, NativeOp->GetOperation()))
+		{
+			return FReply::Handled();
+		}
 	}
 
 	if (OnCommonButtonDragOver().IsBound())
@@ -1650,7 +1653,10 @@ COMMONUI_API FReply UCommonButtonBase::HandleButtonDrop(const FGeometry& MyGeome
 	TSharedPtr<FUMGDragDropOp> NativeOp = DragDropEvent.GetOperationAs<FUMGDragDropOp>();
 	if (NativeOp.IsValid())
 	{
-		NativeOnDrop(MyGeometry, DragDropEvent, NativeOp->GetOperation());
+		if (NativeOnDrop(MyGeometry, DragDropEvent, NativeOp->GetOperation()))
+		{
+			return FReply::Handled();	
+		}
 	}
 
 	if (OnCommonButtonDrop().IsBound())
@@ -1824,7 +1830,6 @@ void UCommonButtonBase::NativeOnReleased()
 
 void UCommonButtonBase::NativeOnDragDetected(const FGeometry& InGeometry, const FPointerEvent& InMouseEvent, UDragDropOperation*& OutOperation)
 {
-	BP_OnDragDetected();
 	if (OnButtonBaseDragDetected.IsBound())
 	{
 		OnButtonBaseDragDetected.Broadcast(this, InGeometry, OutOperation);
@@ -1834,7 +1839,6 @@ void UCommonButtonBase::NativeOnDragDetected(const FGeometry& InGeometry, const
 
 void UCommonButtonBase::NativeOnDragEnter(const FGeometry& InGeometry, const FDragDropEvent& InDragDropEvent, UDragDropOperation* InOperation)
 {
-	BP_OnDragEnter();
 	if (OnButtonBaseDragEnter.IsBound())
 	{
 		OnButtonBaseDragEnter.Broadcast(this, InGeometry, InOperation);
@@ -1844,7 +1848,6 @@ void UCommonButtonBase::NativeOnDragEnter(const FGeometry& InGeometry, const FDr
 
 void UCommonButtonBase::NativeOnDragLeave(const FDragDropEvent& InDragDropEvent, UDragDropOperation* InOperation)
 {
-	BP_OnDragLeave();
 	if (OnButtonBaseDragLeave.IsBound())
 	{
 		OnButtonBaseDragLeave.Broadcast(this, InOperation);
@@ -1854,7 +1857,6 @@ void UCommonButtonBase::NativeOnDragLeave(const FDragDropEvent& InDragDropEvent,
 
 bool UCommonButtonBase::NativeOnDragOver(const FGeometry& InGeometry, const FDragDropEvent& InDragDropEvent, UDragDropOperation* InOperation)
 {
-	BP_OnDragOver();
 	if (OnButtonBaseDragOver.IsBound())
 	{
 		OnButtonBaseDragOver.Broadcast(this,InGeometry, InOperation);
@@ -1864,10 +1866,9 @@ bool UCommonButtonBase::NativeOnDragOver(const FGeometry& InGeometry, const FDra
 
 bool UCommonButtonBase::NativeOnDrop(const FGeometry& InGeometry, const FDragDropEvent& InDragDropEvent, UDragDropOperation* InOperation)
 {
-	BP_OnDrop();
 	if (OnButtonBaseDrop.IsBound())
 	{
-		OnButtonBaseDragOver.Broadcast(this, InGeometry, InOperation);
+		OnButtonBaseDrop.Broadcast(this, InGeometry, InOperation);
 	}
 	return Super::NativeOnDrop(InGeometry, InDragDropEvent, InOperation);
 }
diff --git a/Engine/Plugins/Runtime/CommonUI/Source/CommonUI/Public/CommonButtonBase.h b/Engine/Plugins/Runtime/CommonUI/Source/CommonUI/Public/CommonButtonBase.h
--- a/Engine/Plugins/Runtime/CommonUI/Source/CommonUI/Public/CommonButtonBase.h
+++ b/Engine/Plugins/Runtime/CommonUI/Source/CommonUI/Public/CommonButtonBase.h
@@ -719,24 +719,14 @@ protected:
 	COMMONUI_API void BP_OnReleased();
 	COMMONUI_API virtual void NativeOnReleased();
 
-	UFUNCTION(BlueprintImplementableEvent, Category = CommonButton, meta = (DisplayName = "On Drag Detected"))
-	COMMONUI_API void BP_OnDragDetected();
 	COMMONUI_API virtual void NativeOnDragDetected(const FGeometry& InGeometry, const FPointerEvent& InMouseEvent, UDragDropOperation*& OutOperation) override;
 
-	UFUNCTION(BlueprintImplementableEvent, Category = CommonButton, meta = (DisplayName = "On Drag Enter"))
-	COMMONUI_API void BP_OnDragEnter();
 	COMMONUI_API virtual void NativeOnDragEnter(const FGeometry& InGeometry, const FDragDropEvent& InDragDropEvent, UDragDropOperation* InOperation) override;
-
-	UFUNCTION(BlueprintImplementableEvent, Category = CommonButton, meta = (DisplayName = "On Drag Leave"))
-	COMMONUI_API void BP_OnDragLeave();
+	
 	COMMONUI_API virtual void NativeOnDragLeave(const FDragDropEvent& InDragDropEvent, UDragDropOperation* InOperation) override;
 
-	UFUNCTION(BlueprintImplementableEvent, Category = CommonButton, meta = (DisplayName = "On Drag Over"))
-	COMMONUI_API void BP_OnDragOver();
 	COMMONUI_API virtual bool NativeOnDragOver(const FGeometry& InGeometry, const FDragDropEvent& InDragDropEvent, UDragDropOperation* InOperation) override;
 
-	UFUNCTION(BlueprintImplementableEvent, Category = CommonButton, meta = (DisplayName = "On Drop"))
-	COMMONUI_API void BP_OnDrop();
 	COMMONUI_API virtual bool NativeOnDrop(const FGeometry& InGeometry, const FDragDropEvent& InDragDropEvent, UDragDropOperation* InOperation) override;
 
 	UFUNCTION(BlueprintImplementableEvent, Category = CommonButton, meta = (DisplayName = "On Enabled"))
