CommonUI: Fix issue of accepting widget before it is rendered to under cursor.

#rnx

diff --git a/Engine/Plugins/Runtime/CommonUI/Source/CommonUI/Private/Input/CommonAnalogCursor.cpp b/Engine/Plugins/Runtime/CommonUI/Source/CommonUI/Private/Input/CommonAnalogCursor.cpp
--- a/Engine/Plugins/Runtime/CommonUI/Source/CommonUI/Private/Input/CommonAnalogCursor.cpp
+++ b/Engine/Plugins/Runtime/CommonUI/Source/CommonUI/Private/Input/CommonAnalogCursor.cpp
@@ -27,6 +27,11 @@ static const TAutoConsoleVariable<bool> CVarShouldVirtualAcceptSimulateMouseButt
 	true,
 	TEXT("Controls if virtual_accept key events will be converted to left mouse button events."));
 
+static const TAutoConsoleVariable<bool> CVarShouldRouteOffscreenMouseButton(
+	TEXT("CommonUI.ShouldRouteOffscreenMouseButton"),
+	true,
+	TEXT("Controls if we should directly route mouse events to offscreen widgets."));
+
 bool IsEligibleFakeKeyPointerEvent(const FPointerEvent& PointerEvent)
 {
 	FKey EffectingButton = PointerEvent.GetEffectingButton();
@@ -36,6 +41,78 @@ bool IsEligibleFakeKeyPointerEvent(const FPointerEvent& PointerEvent)
 		&& EffectingButton != EKeys::MiddleMouseButton;
 }
 
+bool IsCursorWithinRenderBounds(TSharedRef<SWidget> Widget, FSlateApplication& SlateApp, TSharedRef<FSlateUser> SlateUser)
+{
+	// Cursor is outside the clipped widget (i.e. outside viewport)
+	TOptional<FSlateClippingState> ClippingState = Widget->GetCurrentClippingState();
+	if (ClippingState.IsSet() && !ClippingState->IsPointInside(SlateUser->GetCursorPosition()))
+	{
+		return false;
+	}
+	else
+	{
+		// Cursor is outside window
+		TSharedPtr<SWindow> FocusedWindow = SlateApp.FindWidgetWindow(Widget);
+		if (FocusedWindow.IsValid())
+		{
+			return true; // Widget is outside window, as there is no window
+		}
+
+		const FGeometry Geometry = FocusedWindow->GetCachedGeometry();
+		const FVector2D LocalPoint = Geometry.AbsoluteToLocal(SlateUser->GetCursorPosition());
+
+		if (LocalPoint.X < 0.0 || LocalPoint.Y < 0.0 || LocalPoint.X > Geometry.GetLocalSize().X || LocalPoint.Y > Geometry.GetLocalSize().Y)
+		{
+			return false;
+		}
+	}
+
+	return true;
+}
+
+TOptional<FPointerEvent> GetSimulatedMouseEventForOffscreenFocusedWidget(const FKeyEvent& InKeyEvent, FSlateApplication& SlateApp,FWidgetPath& OutWidgetPath)
+{
+	TSharedPtr<FSlateUser> SlateUser = SlateApp.GetUser(InKeyEvent);
+	TSharedPtr<SWidget> FocusedWidget = SlateApp.GetUserFocusedWidget(InKeyEvent.GetUserIndex());
+	if(!SlateUser.IsValid() || !FocusedWidget.IsValid())
+	{
+		return TOptional<FPointerEvent>();
+	}
+
+	if (IsCursorWithinRenderBounds(FocusedWidget.ToSharedRef(), SlateApp, SlateUser.ToSharedRef()))
+	{
+		return TOptional<FPointerEvent>();
+	}
+
+	FWidgetPath WidgetPathNoPointer;
+	SlateApp.FindPathToWidget(FocusedWidget.ToSharedRef(), WidgetPathNoPointer);
+
+	TArray<FWidgetAndPointer> WidgetAndPointers;
+	WidgetAndPointers.Reserve(WidgetPathNoPointer.Widgets.Num());
+
+	for (int32 Index = 0; Index < WidgetPathNoPointer.Widgets.Num(); ++Index)
+	{
+		const FVirtualPointerPosition VirtualCursorPosition(SlateUser->GetCursorPosition(), SlateUser->GetPreviousCursorPosition());
+		WidgetAndPointers.Add(FWidgetAndPointer(WidgetPathNoPointer.Widgets[Index], VirtualCursorPosition)) ;
+	}
+
+	OutWidgetPath = FWidgetPath(WidgetAndPointers);
+
+	const bool bIsPrimaryUser = FSlateApplication::CursorUserIndex == SlateUser->GetUserIndex();
+	FPointerEvent MouseEvent(
+		SlateUser->GetUserIndex(),
+		FSlateApplication::CursorPointerIndex,
+		SlateUser->GetCursorPosition(),
+		SlateUser->GetPreviousCursorPosition(),
+		bIsPrimaryUser ? SlateApp.GetPressedMouseButtons() : FTouchKeySet::EmptySet,
+		EKeys::LeftMouseButton,
+		0,
+		bIsPrimaryUser ? SlateApp.GetModifierKeys() : FModifierKeysState()
+	);
+
+	return MouseEvent;
+}
+
 FCommonAnalogCursor::FCommonAnalogCursor(const UCommonUIActionRouterBase& InActionRouter)
 	: ActionRouter(InActionRouter)
 	, ActiveInputMethod(ECommonInputType::MouseAndKeyboard)
@@ -263,6 +340,23 @@ bool FCommonAnalogCursor::HandleKeyDownEvent(FSlateApplication& SlateApp, const
 		}
 		else if (!bIsVirtualAccept || ShouldVirtualAcceptSimulateMouseButton(InKeyEvent, IE_Pressed))
 		{
+			// If virtually accepting the focused widget, and cursor is not within the current cliprect or window, we can not rely on the hittest grid as the cursor and widget position may not updated correctly.
+			// For instance, when attemptign to forward mouse input to animating or offscreen widgets. As a workaround, directly forard the pointer down event to the focused widget path.
+			if (bIsVirtualAccept && CVarShouldRouteOffscreenMouseButton.GetValueOnGameThread())
+			{
+				FWidgetPath FocusedWidgetPath;
+				TOptional<FPointerEvent> MouseEvent = GetSimulatedMouseEventForOffscreenFocusedWidget(InKeyEvent, SlateApp, FocusedWidgetPath);
+				if (MouseEvent.IsSet())
+				{
+					UCommonInputSubsystem& InputSubsystem = ActionRouter.GetInputSubsystem();
+					InputSubsystem.SetIsGamepadSimulatedClick(bIsVirtualAccept);
+					const bool bReturnValue = SlateApp.RoutePointerDownEvent(FocusedWidgetPath, MouseEvent.GetValue()).IsEventHandled();
+					InputSubsystem.SetIsGamepadSimulatedClick(false);
+					
+					return bReturnValue;
+				}
+			}
+
 			// There is no awareness on a mouse event of whether it's real or not, so mark that here.
 			UCommonInputSubsystem& InputSubsystem = ActionRouter.GetInputSubsystem();
 			InputSubsystem.SetIsGamepadSimulatedClick(bIsVirtualAccept);
@@ -297,6 +391,18 @@ bool FCommonAnalogCursor::HandleKeyUpEvent(FSlateApplication& SlateApp, const FK
 		}
 		else if (!bIsVirtualAccept || ShouldVirtualAcceptSimulateMouseButton(InKeyEvent, IE_Released))
 		{
+						// If virtually accepting the focused widget, and cursor is not within the current cliprect or window, we can not rely on the hittest grid as the cursor and widget position may not updated correctly.
+			// For instance, when attemptign to forward mouse input to animating or offscreen widgets. As a workaround, directly forard the pointer down event to the focused widget path.
+			if (bIsVirtualAccept && CVarShouldRouteOffscreenMouseButton.GetValueOnGameThread())
+			{
+				FWidgetPath FocusedWidgetPath;
+				TOptional<FPointerEvent> MouseEvent = GetSimulatedMouseEventForOffscreenFocusedWidget(InKeyEvent, SlateApp, FocusedWidgetPath);
+				if (MouseEvent.IsSet())
+				{
+					return SlateApp.RoutePointerUpEvent(FocusedWidgetPath, MouseEvent.GetValue()).IsEventHandled();
+				}
+			}
+
 			return FAnalogCursor::HandleKeyUpEvent(SlateApp, InKeyEvent);
 		}
 	}
