CommonUI: Fixup CL 44548495

#rnx

diff --git a/Engine/Plugins/Runtime/CommonUI/Source/CommonUI/Private/Input/CommonAnalogCursor.cpp b/Engine/Plugins/Runtime/CommonUI/Source/CommonUI/Private/Input/CommonAnalogCursor.cpp
--- a/Engine/Plugins/Runtime/CommonUI/Source/CommonUI/Private/Input/CommonAnalogCursor.cpp
+++ b/Engine/Plugins/Runtime/CommonUI/Source/CommonUI/Private/Input/CommonAnalogCursor.cpp
@@ -53,7 +53,7 @@ bool IsCursorWithinRenderBounds(TSharedRef<SWidget> Widget, FSlateApplication& S
 	{
 		// Cursor is outside window
 		TSharedPtr<SWindow> FocusedWindow = SlateApp.FindWidgetWindow(Widget);
-		if (FocusedWindow.IsValid())
+		if (!FocusedWindow.IsValid())
 		{
 			return true; // Widget is outside window, as there is no window
 		}
@@ -87,6 +87,11 @@ TOptional<FPointerEvent> GetSimulatedMouseEventForOffscreenFocusedWidget(const F
 	FWidgetPath WidgetPathNoPointer;
 	SlateApp.FindPathToWidget(FocusedWidget.ToSharedRef(), WidgetPathNoPointer);
 
+	if (!WidgetPathNoPointer.IsValid())
+	{
+		return TOptional<FPointerEvent>();
+	}
+
 	TArray<FWidgetAndPointer> WidgetAndPointers;
 	WidgetAndPointers.Reserve(WidgetPathNoPointer.Widgets.Num());
 
