Add some UI logging to help track down issues with client losing focus when alt tabbing during a reboot (as the player being rebooted)

diff --git a/Engine/Plugins/Runtime/CommonUI/Source/CommonUI/Private/Widgets/CommonActivatableWidgetContainer.cpp b/Engine/Plugins/Runtime/CommonUI/Source/CommonUI/Private/Widgets/CommonActivatableWidgetContainer.cpp
--- a/Engine/Plugins/Runtime/CommonUI/Source/CommonUI/Private/Widgets/CommonActivatableWidgetContainer.cpp
+++ b/Engine/Plugins/Runtime/CommonUI/Source/CommonUI/Private/Widgets/CommonActivatableWidgetContainer.cpp
@@ -270,7 +270,7 @@ void UCommonActivatableWidgetContainerBase::ReleaseWidget(const TSharedRef<SWidg
 
 void UCommonActivatableWidgetContainerBase::HandleActiveIndexChanged(int32 ActiveWidgetIndex)
 {
-	UE_LOG(LogCommonUI, VeryVerbose, TEXT("UCommonActivatableWidgetContainerBase::HandleActiveIndexChanged() ActiveWidgetIndex: %d"), ActiveWidgetIndex);
+	UE_LOG(LogCommonUI, Log, TEXT("UCommonActivatableWidgetContainerBase::HandleActiveIndexChanged() ActiveWidgetIndex: %d"), ActiveWidgetIndex);
 	
 	// Remove all slots above the currently active one and release the widgets back to the pool
 	while (MySwitcher->GetNumWidgets() - 1 > ActiveWidgetIndex)
