Prevent crash on templated commonvideoplayer getting put in widget blueprints

#tests Forced the pointer to be wrong to get the same callstack. Tried to find repro steps but couldn't figure out how to cause the object to be considered a Template
#rnx

diff --git a/Engine/Plugins/Runtime/CommonUI/Source/CommonUI/Private/CommonVideoPlayer.cpp b/Engine/Plugins/Runtime/CommonUI/Source/CommonUI/Private/CommonVideoPlayer.cpp
--- a/Engine/Plugins/Runtime/CommonUI/Source/CommonUI/Private/CommonVideoPlayer.cpp
+++ b/Engine/Plugins/Runtime/CommonUI/Source/CommonUI/Private/CommonVideoPlayer.cpp
@@ -176,6 +176,10 @@ void UCommonVideoPlayer::SynchronizeProperties()
 {
 	Super::SynchronizeProperties();
 	
+	if (MediaPlayer == nullptr)
+	{
+		return;
+	}
 	SetVideo(Video);
 }
 
