Common Ui Action Router:
- Refresh root nodes when calling RefreshActionDomainLeafNodeConfig

#rnx

diff --git a/Engine/Plugins/Runtime/CommonUI/Source/CommonUI/Private/Input/CommonUIActionRouterBase.cpp b/Engine/Plugins/Runtime/CommonUI/Source/CommonUI/Private/Input/CommonUIActionRouterBase.cpp
--- a/Engine/Plugins/Runtime/CommonUI/Source/CommonUI/Private/Input/CommonUIActionRouterBase.cpp
+++ b/Engine/Plugins/Runtime/CommonUI/Source/CommonUI/Private/Input/CommonUIActionRouterBase.cpp
@@ -1156,6 +1156,12 @@ void UCommonUIActionRouterBase::DebugDumpActionDomainRootNodes(int32 UserIndex,
 }
 
 void UCommonUIActionRouterBase::ProcessRebuiltWidgets()
+{
+	RefreshRootNodes();
+	RefreshBoundActions();
+}
+
+void UCommonUIActionRouterBase::RefreshRootNodes()
 {
 	// Begin by organizing all of the widgets that need nodes according to their direct parent
 	TArray<UCommonActivatableWidget*> RootCandidates;
@@ -1175,6 +1181,7 @@ void UCommonUIActionRouterBase::ProcessRebuiltWidgets()
 			}
 		}
 	}
+	RebuiltWidgetsPendingNodeAssignment.Reset();
 
 	// Build a new tree for any new roots
 	for (UCommonActivatableWidget* RootWidget : RootCandidates)
@@ -1228,7 +1235,10 @@ void UCommonUIActionRouterBase::ProcessRebuiltWidgets()
 	{
 		ensureAlwaysMsgf(false, TEXT("Somehow we rebuilt a widget that is owned by an activatable, but no node exists for that activatable. This *should* be completely impossible."));
 	}
+}
 	
+void UCommonUIActionRouterBase::RefreshBoundActions()
+{
 	// Now, we account for all the widgets that would like their actions bound
 	for (const FPendingWidgetRegistration& PendingRegistration : PendingWidgetRegistrations)
 	{
@@ -1258,7 +1268,6 @@ void UCommonUIActionRouterBase::ProcessRebuiltWidgets()
 		}
 	}
 
-	RebuiltWidgetsPendingNodeAssignment.Reset();
 	PendingWidgetRegistrations.Reset();
 }
 
@@ -1756,6 +1765,11 @@ void UCommonUIActionRouterBase::RefreshActionDomainLeafNodeConfig()
 
 	if (const UCommonInputActionDomainTable* ActionDomainTable = GetActionDomainTable())
 	{
+		if (RebuiltWidgetsPendingNodeAssignment.Num() > 0)
+		{
+			RefreshRootNodes();
+		}
+
 		if (FActivatableTreeRootPtr RootNode = FindActiveActionDomainRootNode())
 		{
 			if (!RootNode->UpdateLeafmostActiveNode(RootNode))
diff --git a/Engine/Plugins/Runtime/CommonUI/Source/CommonUI/Public/Input/CommonUIActionRouterBase.h b/Engine/Plugins/Runtime/CommonUI/Source/CommonUI/Public/Input/CommonUIActionRouterBase.h
--- a/Engine/Plugins/Runtime/CommonUI/Source/CommonUI/Public/Input/CommonUIActionRouterBase.h
+++ b/Engine/Plugins/Runtime/CommonUI/Source/CommonUI/Public/Input/CommonUIActionRouterBase.h
@@ -186,6 +186,8 @@ private:
 	
 	UE_API void HandleActivatableWidgetRebuilding(UCommonActivatableWidget& RebuildingWidget);
 	UE_API void ProcessRebuiltWidgets();
+	UE_API void RefreshRootNodes();
+	UE_API void RefreshBoundActions();
 	UE_API void AssembleTreeRecursive(const FActivatableTreeNodeRef& CurNode, TMap<UCommonActivatableWidget*, TArray<UCommonActivatableWidget*>>& WidgetsByDirectParent);
 
 	UE_API void HandleRootWidgetSlateReleased(TWeakPtr<FActivatableTreeRoot> WeakRoot);
