Added TSoftClassPtr 'WidgetClass' UPROPERTY to UCommonLazyWidget.

This allows us to design-time preview the widget placement in the editor.

Additionally added 'LoadLazyContent' method which takes a TFunction parameter for easy lambda-like post-load initialization logic.

#tests Compiled DevelopmentEditor Win64 / Test Client / BugHawk / Ran a PIE instance and verified lazy load, usage, and closing of modal dialogs. Repeated test several times on each dialog to ensure they would be correctly cleaned up and instantiated again.

diff --git a/Engine/Plugins/Runtime/CommonUI/Source/CommonUI/Private/CommonLazyWidget.cpp b/Engine/Plugins/Runtime/CommonUI/Source/CommonUI/Private/CommonLazyWidget.cpp
--- a/Engine/Plugins/Runtime/CommonUI/Source/CommonUI/Private/CommonLazyWidget.cpp
+++ b/Engine/Plugins/Runtime/CommonUI/Source/CommonUI/Private/CommonLazyWidget.cpp
@@ -16,6 +16,16 @@ UCommonLazyWidget::UCommonLazyWidget(const FObjectInitializer& Initializer)
 	SetVisibilityInternal(ESlateVisibility::SelfHitTestInvisible);
 }
 
+void UCommonLazyWidget::SetLazyContent(const TSoftClassPtr<UUserWidget> SoftWidget)
+{
+	SetLazyContentInternal(SoftWidget, {});
+}
+
+void UCommonLazyWidget::LoadLazyContent()
+{
+	LoadLazyContent({});
+}
+
 TSharedRef<SWidget> UCommonLazyWidget::RebuildWidget()
 {
 	MyLoadGuard = SNew(SLoadGuard)
@@ -50,7 +60,14 @@ void UCommonLazyWidget::SynchronizeProperties()
 #if WITH_EDITOR
 		if (IsDesignTime())
 		{
-			MyLoadGuard->SetForceShowSpinner(true);
+			if (!WidgetClass.IsNull())
+			{
+				SetLazyContent(WidgetClass);
+			}
+			else
+			{
+				MyLoadGuard->SetForceShowSpinner(true);
+			}
 		}
 #endif
 	}
@@ -93,42 +110,16 @@ bool UCommonLazyWidget::IsLoading() const
 	return MyLoadGuard.IsValid() && MyLoadGuard->IsLoading();
 }
 
-void UCommonLazyWidget::SetLazyContent(const TSoftClassPtr<UUserWidget> SoftWidget)
-{
-	if (SoftWidget.IsNull())
-	{
-		CancelStreaming();
-		SetLoadedContent(nullptr);
-	}
-
-	TWeakObjectPtr<UCommonLazyWidget> WeakThis(this);
-
-	RequestAsyncLoad(SoftWidget,
-		[WeakThis, SoftWidget]() {
-			if (ThisClass* StrongThis = WeakThis.Get())
-			{
-				if (ensureMsgf(SoftWidget.Get(), TEXT("Failed to load %s"), *SoftWidget.ToSoftObjectPath().ToString()))
-				{
-					// Don't reload the class if we're already this class.
-					if (StrongThis->Content && StrongThis->Content->GetClass() == SoftWidget.Get())
-					{
-						StrongThis->OnContentChangedEvent.Broadcast(StrongThis->Content);
-						return;
-					}
-
-					UUserWidget* UserWidget = CreateWidget(StrongThis->GetOwningPlayer(), SoftWidget.Get());
-					StrongThis->SetLoadedContent(UserWidget);
-				}
-			}
-		}
-	);
-}
-
 void UCommonLazyWidget::SetLoadedContent(UUserWidget* InContent)
 {
-	if (UCommonActivatableWidget* OutgoingActivatable = Cast<UCommonActivatableWidget>(Content))
+	const bool bPerformActivateDeactivate = !IsDesignTime();
+
+	if (bPerformActivateDeactivate)
 	{
-		OutgoingActivatable->DeactivateWidget();
+		if (UCommonActivatableWidget* OutgoingActivatable = Cast<UCommonActivatableWidget>(Content))
+		{
+			OutgoingActivatable->DeactivateWidget();
+		}
 	}
 
 	Content = InContent;
@@ -138,9 +129,12 @@ void UCommonLazyWidget::SetLoadedContent(UUserWidget* InContent)
 		MyLoadGuard->SetContent(InContent ? InContent->TakeWidget() : SNullWidget::NullWidget);
 	}
 
-	if (UCommonActivatableWidget* IncomingActivatable = Cast<UCommonActivatableWidget>(Content))
+	if (bPerformActivateDeactivate)
 	{
-		IncomingActivatable->ActivateWidget();
+		if (UCommonActivatableWidget* IncomingActivatable = Cast<UCommonActivatableWidget>(Content))
+		{
+			IncomingActivatable->ActivateWidget();
+		}
 	}
 
 	OnContentChangedEvent.Broadcast(Content);
diff --git a/Engine/Plugins/Runtime/CommonUI/Source/CommonUI/Public/CommonLazyWidget.h b/Engine/Plugins/Runtime/CommonUI/Source/CommonUI/Public/CommonLazyWidget.h
--- a/Engine/Plugins/Runtime/CommonUI/Source/CommonUI/Public/CommonLazyWidget.h
+++ b/Engine/Plugins/Runtime/CommonUI/Source/CommonUI/Public/CommonLazyWidget.h
@@ -4,8 +4,8 @@
 
 #include "Blueprint/UserWidget.h"
 #include "CommonLoadGuard.h"
-
 #include "Engine/StreamableManager.h"
+
 #include "CommonLazyWidget.generated.h"
 
 #define UE_API COMMONUI_API
@@ -15,7 +15,7 @@ class UCommonMcpItemDefinition;
 DECLARE_MULTICAST_DELEGATE_OneParam(FOnLazyContentChangedEvent, UUserWidget*);
 
 /**
- * A special Image widget that can show unloaded images and takes care of the loading for you!
+ * A widget that can async load and create an instance of a UserWidget.
  */
 UCLASS(MinimalAPI)
 class UCommonLazyWidget : public UWidget
@@ -23,14 +23,33 @@ class UCommonLazyWidget : public UWidget
 	GENERATED_UCLASS_BODY()
 
 public:
-	/**  */
+	/** Loads and creates an instance of SoftWidget. */
 	UFUNCTION(BlueprintCallable, Category = LazyContent)
 	UE_API void SetLazyContent(const TSoftClassPtr<UUserWidget> SoftWidget);
 
-	/**  */
+	/** Loads and creates an instance of the WidgetClass property. */
+	UFUNCTION(BlueprintCallable, Category = LazyContent)
+	UE_API void LoadLazyContent();
+
+	template<typename TWidget = UUserWidget>
+	void LoadLazyContent(TFunction<void(TWidget&)>&& InitInstanceFunc)
+	{
+		if (!WidgetClass.IsNull())
+		{
+			SetLazyContentInternal(WidgetClass, MoveTemp(InitInstanceFunc));
+		}
+	}
+
+	/** Gets the attached Content which was instanced from an async loaded TSoftClassPtr. */
 	UFUNCTION(BlueprintCallable, Category = LazyContent)
 	UUserWidget* GetContent() const { return Content; }
 
+	template <class TContent = UUserWidget>
+	TContent* GetContent() const
+	{
+		return Cast<TContent>(GetContent());
+	}
+
 	UFUNCTION(BlueprintCallable, Category = LazyContent)
 	UE_API bool IsLoading() const;
 
@@ -38,6 +57,43 @@ public:
 	FOnLoadGuardStateChangedEvent& OnLoadingStateChanged() { return OnLoadingStateChangedEvent; }
 
 protected:
+	template <typename TWidget = UUserWidget>
+	void SetLazyContentInternal(const TSoftClassPtr<UUserWidget> SoftWidget, TFunction<void(TWidget&)>&& InitInstanceFunc)
+	{
+		if (SoftWidget.IsNull())
+		{
+			CancelStreaming();
+			SetLoadedContent(nullptr);
+			return;
+		}
+
+		TWeakObjectPtr<UCommonLazyWidget> WeakThis(this);
+
+		RequestAsyncLoad(SoftWidget,
+						 [WeakThis, SoftWidget, InitInstanceFunc = MoveTemp(InitInstanceFunc)]() {
+			if (ThisClass* StrongThis = WeakThis.Get())
+			{
+				if (ensureMsgf(SoftWidget.Get(), TEXT("Failed to load %s"), *SoftWidget.ToSoftObjectPath().ToString()))
+				{
+					// Don't reload the class if we're already this class.
+					if (StrongThis->Content && StrongThis->Content->GetClass() == SoftWidget.Get())
+					{
+						return;
+					}
+
+					TWidget* UserWidget = CreateWidget<TWidget>(StrongThis, SoftWidget.Get());
+
+					if (InitInstanceFunc && UserWidget && !StrongThis->IsDesignTime())
+					{
+						InitInstanceFunc(*UserWidget);
+					}
+
+					StrongThis->SetLoadedContent(UserWidget);
+				}
+			}
+		});
+	}
+
 	UE_API virtual TSharedRef<SWidget> RebuildWidget() override;
 	UE_API virtual void OnWidgetRebuilt() override;
 	UE_API virtual void ReleaseSlateResources(bool bReleaseChildren) override;
@@ -59,6 +115,9 @@ private:
 	UE_API void RequestAsyncLoad(TSoftClassPtr<UObject> SoftObject, FStreamableDelegate DelegateToCall);
 	UE_API void HandleLoadGuardStateChanged(bool bIsLoading);
 
+	UPROPERTY(EditAnywhere, Category = LazyWidget)
+	TSoftClassPtr<UUserWidget> WidgetClass;
+
 	/** The loading throbber brush */
 	UPROPERTY(EditAnywhere, Category = Appearance)
 	FSlateBrush LoadingThrobberBrush;
@@ -72,7 +131,7 @@ private:
 	TSharedPtr<FStreamableHandle> StreamingHandle;
 	FSoftObjectPath StreamingObjectPath;
 
-	UPROPERTY(BlueprintAssignable, Category = LazyImage, meta = (DisplayName = "On Loading State Changed", ScriptName = "OnLoadingStateChanged"))
+	UPROPERTY(BlueprintAssignable, Category = LazyWidget, meta = (DisplayName = "On Loading State Changed", ScriptName = "OnLoadingStateChanged"))
 	FOnLoadGuardStateChangedDynamic BP_OnLoadingStateChanged;
 
 	TSharedPtr<SLoadGuard> MyLoadGuard;
