PR #13800: Fixed engine crash in CommonUI.DumpActivatableTree console command

#rnx

diff --git a/Engine/Plugins/Runtime/CommonUI/Source/CommonUI/Private/Input/CommonUIActionRouterBase.cpp b/Engine/Plugins/Runtime/CommonUI/Source/CommonUI/Private/Input/CommonUIActionRouterBase.cpp
--- a/Engine/Plugins/Runtime/CommonUI/Source/CommonUI/Private/Input/CommonUIActionRouterBase.cpp
+++ b/Engine/Plugins/Runtime/CommonUI/Source/CommonUI/Private/Input/CommonUIActionRouterBase.cpp
@@ -1981,6 +1981,7 @@ public:
 	{
 		if (World == nullptr)
 		{
+			UE_LOG(LogUIActionRouter, Error, TEXT("No World, unable to run CommonUI.DumpActivatableTree"));
 			return;
 		}
 
@@ -1990,6 +1991,12 @@ public:
 		const int LocalPlayerIndex = Args.IsValidIndex(3) ? FCString::Atoi(*Args[3]) : -1;
 
 		UGameInstance* GameInstance = World->GetGameInstance();
+		if (GameInstance == nullptr)
+		{
+			UE_LOG(LogUIActionRouter, Error, TEXT("No GameInstance, unable to run CommonUI.DumpActivatableTree"));
+			return;
+		}
+
 		const TArray<ULocalPlayer*>& LocalPlayers = GameInstance->GetLocalPlayers();
 		for (int32 CurrIdx = 0; CurrIdx < LocalPlayers.Num(); ++CurrIdx)
 		{
