[Core] Add SetLazyContentWithCallback to Lazy Widget

#tests Mostly just validated signature

diff --git a/Engine/Plugins/Runtime/CommonUI/Source/CommonUI/Private/CommonLazyWidget.cpp b/Engine/Plugins/Runtime/CommonUI/Source/CommonUI/Private/CommonLazyWidget.cpp
--- a/Engine/Plugins/Runtime/CommonUI/Source/CommonUI/Private/CommonLazyWidget.cpp
+++ b/Engine/Plugins/Runtime/CommonUI/Source/CommonUI/Private/CommonLazyWidget.cpp
@@ -21,6 +21,11 @@ void UCommonLazyWidget::SetLazyContent(const TSoftClassPtr<UUserWidget> SoftWidg
 	SetLazyContentInternal(SoftWidget, {});
 }
 
+void UCommonLazyWidget::SetLazyContentWithCallback(const TSoftClassPtr<UUserWidget> SoftWidget, const FOnWidgetCreated& OnCreatedCallback)
+{
+	SetLazyContentInternal<UUserWidget>(SoftWidget, [OnCreatedCallback](UUserWidget& Widget) { OnCreatedCallback.ExecuteIfBound(&Widget); });
+}
+
 void UCommonLazyWidget::LoadLazyContent()
 {
 	LoadLazyContent({});
diff --git a/Engine/Plugins/Runtime/CommonUI/Source/CommonUI/Public/CommonLazyWidget.h b/Engine/Plugins/Runtime/CommonUI/Source/CommonUI/Public/CommonLazyWidget.h
--- a/Engine/Plugins/Runtime/CommonUI/Source/CommonUI/Public/CommonLazyWidget.h
+++ b/Engine/Plugins/Runtime/CommonUI/Source/CommonUI/Public/CommonLazyWidget.h
@@ -23,10 +23,16 @@ class UCommonLazyWidget : public UWidget
 	GENERATED_UCLASS_BODY()
 
 public:
+	DECLARE_DYNAMIC_DELEGATE_OneParam(FOnWidgetCreated, UUserWidget*, Widget);
+
 	/** Loads and creates an instance of SoftWidget. */
 	UFUNCTION(BlueprintCallable, Category = LazyContent)
 	UE_API void SetLazyContent(const TSoftClassPtr<UUserWidget> SoftWidget);
 
+	/** Loads and creates an instance of SoftWidget. */
+	UFUNCTION(BlueprintCallable, Category = LazyContent)
+	UE_API void SetLazyContentWithCallback(const TSoftClassPtr<UUserWidget> SoftWidget, const FOnWidgetCreated& OnCreatedCallback);
+
 	/** Loads and creates an instance of the WidgetClass property. */
 	UFUNCTION(BlueprintCallable, Category = LazyContent)
 	UE_API void LoadLazyContent();
