Disable ShouldRouteOffscreenMouseButton

#rnx

diff --git a/Engine/Plugins/Runtime/CommonUI/Source/CommonUI/Private/Input/CommonAnalogCursor.cpp b/Engine/Plugins/Runtime/CommonUI/Source/CommonUI/Private/Input/CommonAnalogCursor.cpp
--- a/Engine/Plugins/Runtime/CommonUI/Source/CommonUI/Private/Input/CommonAnalogCursor.cpp
+++ b/Engine/Plugins/Runtime/CommonUI/Source/CommonUI/Private/Input/CommonAnalogCursor.cpp
@@ -29,7 +29,7 @@ static const TAutoConsoleVariable<bool> CVarShouldVirtualAcceptSimulateMouseButt
 
 static const TAutoConsoleVariable<bool> CVarShouldRouteOffscreenMouseButton(
 	TEXT("CommonUI.ShouldRouteOffscreenMouseButton"),
-	true,
+	false,
 	TEXT("Controls if we should directly route mouse events to offscreen widgets."));
 
 bool IsEligibleFakeKeyPointerEvent(const FPointerEvent& PointerEvent)
