CommonUI: Warn on duplicate tabs, do not ensure.

#rnx

diff --git a/Engine/Plugins/Runtime/CommonUI/Source/CommonUI/Private/CommonTabListWidgetBase.cpp b/Engine/Plugins/Runtime/CommonUI/Source/CommonUI/Private/CommonTabListWidgetBase.cpp
--- a/Engine/Plugins/Runtime/CommonUI/Source/CommonUI/Private/CommonTabListWidgetBase.cpp
+++ b/Engine/Plugins/Runtime/CommonUI/Source/CommonUI/Private/CommonTabListWidgetBase.cpp
@@ -3,6 +3,7 @@
 #include "CommonTabListWidgetBase.h"
 
 #include "CommonAnimatedSwitcher.h"
+#include "CommonUIPrivate.h"
 #include "Groups/CommonButtonGroupBase.h"
 #include "Input/CommonUIInputTypes.h"
 #include "CommonUITypes.h"
@@ -44,22 +45,25 @@ bool UCommonTabListWidgetBase::RegisterTab(FName TabNameID, TSubclassOf<UCommonB
 	bool bAreParametersValid = true;
 
 	// Early out on redundant tab registration.
-	if (!ensure(!RegisteredTabsByID.Contains(TabNameID)))
+	if (RegisteredTabsByID.Contains(TabNameID))
 	{
 		bAreParametersValid = false;
+		UE_LOG(LogCommonUI, Warning, TEXT("RegisteredTabsByID already contains a tab called [%s]"), *TabNameID.ToString());
 	}
 
 	// Early out on invalid tab button type.
-	if (!ensure(ButtonWidgetType))
+	if (ButtonWidgetType == nullptr)
 	{
 		bAreParametersValid = false;
+		UE_LOG(LogCommonUI, Warning, TEXT("RegisteredTabsByID missing ButtonWidgetType [%s]"), *TabNameID.ToString());
 	}
 	
 	// NOTE: Adding the button to the group may change it's selection, which raises an event we listen to,
 	// which can only properly be handled if we already know that this button is associated with a registered tab.
-	if (!ensure(TabButtonGroup))
+	if (TabButtonGroup == nullptr)
 	{
 		bAreParametersValid = false;
+		UE_LOG(LogCommonUI, Warning, TEXT("RegisteredTabsByID missing TabButtonGroup [%s]"), *TabNameID.ToString());
 	}
 
 	if (!bAreParametersValid)
