CommonButtonBase: avoid calling focus blueprint events during object loading

diff --git a/Engine/Plugins/Runtime/CommonUI/Source/CommonUI/Private/CommonButtonBase.cpp b/Engine/Plugins/Runtime/CommonUI/Source/CommonUI/Private/CommonButtonBase.cpp
--- a/Engine/Plugins/Runtime/CommonUI/Source/CommonUI/Private/CommonButtonBase.cpp
+++ b/Engine/Plugins/Runtime/CommonUI/Source/CommonUI/Private/CommonButtonBase.cpp
@@ -1506,22 +1506,30 @@ void UCommonButtonBase::HandleFocusReceived()
 		SetIsSelected(true, false);
 	}
 	OnFocusReceived().Broadcast();
-	BP_OnFocusReceived();
 
-	if (OnButtonBaseFocused.IsBound())
+	if (CanSafelyRouteCall())
 	{
-		OnButtonBaseFocused.Broadcast(this);
+		BP_OnFocusReceived();
+
+		if (OnButtonBaseFocused.IsBound())
+		{
+			OnButtonBaseFocused.Broadcast(this);
+		}
 	}
 }
 
 void UCommonButtonBase::HandleFocusLost()
 {
 	OnFocusLost().Broadcast();
-	BP_OnFocusLost();
 
-	if (OnButtonBaseUnfocused.IsBound())
+	if (CanSafelyRouteCall())
 	{
-		OnButtonBaseUnfocused.Broadcast(this);
+		BP_OnFocusLost();
+
+		if (OnButtonBaseUnfocused.IsBound())
+		{
+			OnButtonBaseUnfocused.Broadcast(this);
+		}
 	}
 }
 
