Add more logging to RegisterTab when a key is being reused

#tests Added same tab id twice in a row to validate the logged string
#rnx

diff --git a/Engine/Plugins/Runtime/CommonUI/Source/CommonUI/Private/CommonTabListWidgetBase.cpp b/Engine/Plugins/Runtime/CommonUI/Source/CommonUI/Private/CommonTabListWidgetBase.cpp
--- a/Engine/Plugins/Runtime/CommonUI/Source/CommonUI/Private/CommonTabListWidgetBase.cpp
+++ b/Engine/Plugins/Runtime/CommonUI/Source/CommonUI/Private/CommonTabListWidgetBase.cpp
@@ -44,9 +44,10 @@ bool UCommonTabListWidgetBase::RegisterTab(FName TabNameID, TSubclassOf<UCommonB
 	bool bAreParametersValid = true;
 
 	// Early out on redundant tab registration.
-	if (!ensure(!RegisteredTabsByID.Contains(TabNameID)))
+	if (RegisteredTabsByID.Contains(TabNameID))
 	{
 		bAreParametersValid = false;
+		ensureMsgf(bAreParametersValid, TEXT("RegisteredTabsByID already contains a tab called [%s]"), *TabNameID.ToString());
 	}
 
 	// Early out on invalid tab button type.
