Fix missing unregistration of gameplay tag callbacks with merged modules

Delegate instances in GFPs always need to be cleared before that GFP unmounts, or with merged modules we'll crash when the unloaded code is accessed.

#rnx

diff --git a/Engine/Plugins/Runtime/CommonUI/Source/CommonUI/Public/UITag.h b/Engine/Plugins/Runtime/CommonUI/Source/CommonUI/Public/UITag.h
--- a/Engine/Plugins/Runtime/CommonUI/Source/CommonUI/Public/UITag.h
+++ b/Engine/Plugins/Runtime/CommonUI/Source/CommonUI/Public/UITag.h
@@ -73,13 +73,20 @@ class TTypedTagStaticImpl
 	TTypedTagStaticImpl()
 	{ 
 		LLM_SCOPE(ELLMTag::UI);
-		UGameplayTagsManager::OnLastChanceToAddNativeTags().AddLambda([this]()
+		OnLastChanceToAddNativeTagsHandle = UGameplayTagsManager::OnLastChanceToAddNativeTags().AddLambda([this]()
 			{
 				StaticImpl.RootTag = UGameplayTagsManager::Get().AddNativeGameplayTag(TagT::GetRootTagStr());
 			});
 	}
+
+	~TTypedTagStaticImpl()
+	{ 
+		UGameplayTagsManager::OnLastChanceToAddNativeTags().Remove(OnLastChanceToAddNativeTagsHandle);
+	}
+
 	TagT RootTag;
 	static TTypedTagStaticImpl StaticImpl;
+	FDelegateHandle OnLastChanceToAddNativeTagsHandle;
 };
 
 template <typename TagT>
diff --git a/Engine/Source/Runtime/GameplayTags/Classes/GameplayTagContainer.h b/Engine/Source/Runtime/GameplayTags/Classes/GameplayTagContainer.h
--- a/Engine/Source/Runtime/GameplayTags/Classes/GameplayTagContainer.h
+++ b/Engine/Source/Runtime/GameplayTags/Classes/GameplayTagContainer.h
@@ -676,6 +676,8 @@ struct FGameplayTagNativeAdder
 {
 	GAMEPLAYTAGS_API FGameplayTagNativeAdder();
 
+	virtual GAMEPLAYTAGS_API ~FGameplayTagNativeAdder();
+
 	virtual void AddTags() = 0;
 };
 
diff --git a/Engine/Source/Runtime/GameplayTags/Private/GameplayTagContainer.cpp b/Engine/Source/Runtime/GameplayTags/Private/GameplayTagContainer.cpp
--- a/Engine/Source/Runtime/GameplayTags/Private/GameplayTagContainer.cpp
+++ b/Engine/Source/Runtime/GameplayTags/Private/GameplayTagContainer.cpp
@@ -1693,6 +1693,11 @@ FGameplayTagNativeAdder::FGameplayTagNativeAdder()
 	UGameplayTagsManager::OnLastChanceToAddNativeTags().AddRaw(this, &FGameplayTagNativeAdder::AddTags);
 }
 
+FGameplayTagNativeAdder::~FGameplayTagNativeAdder()
+{
+	UGameplayTagsManager::OnLastChanceToAddNativeTags().RemoveAll(this);
+}
+
 FGameplayTagQuery::FGameplayTagQuery()
 	: TokenStreamVersion(EGameplayTagQueryStreamVersion::LatestVersion)
 {
