UI: Fix possible problems with inputs after switching from touch to gamepad in some screens.

The problem is caused by the currently focused widget getting out of sync with the cursor position (even if this cursor is invisible). After we touch the screen in an empty space, cursor is moved there. Then, when we press the confirmation button on a gamepad, it generates mouse click even for this last touched position. It sends the click event to the empty space and some screens lose focus, which leads to semi-lock where gamepad controls cannot be used.

#rnx

diff --git a/Engine/Plugins/Runtime/CommonUI/Source/CommonUI/Private/Input/UIActionRouterTypes.cpp b/Engine/Plugins/Runtime/CommonUI/Source/CommonUI/Private/Input/UIActionRouterTypes.cpp
--- a/Engine/Plugins/Runtime/CommonUI/Source/CommonUI/Private/Input/UIActionRouterTypes.cpp
+++ b/Engine/Plugins/Runtime/CommonUI/Source/CommonUI/Private/Input/UIActionRouterTypes.cpp
@@ -28,6 +28,12 @@ static const TAutoConsoleVariable<bool> CVarGamepadFocusHoveredWidget(
 	false,
 	TEXT("If enabled, the mouse hovered widget, will be focused when switching to gamepad input"));
 
+static const TAutoConsoleVariable<bool> CVarSyncCursorToFocusedWidgetOnInputMethodChange(
+	TEXT("CommonUI.SyncCursorToFocusedWidgetOnInputMethodChange"),
+	true,
+	TEXT("Whether or not to move the cursor to the currently focused widget when input method changes to gamepad (fix for losing focus when switching between touch and gamepad inputs)."),
+	ECVF_Default);
+
 
 const TCHAR* InputEventToString(EInputEvent InputEvent)
 {
@@ -1523,6 +1529,30 @@ void FActivatableTreeRoot::HandleInputMethodChanged(ECommonInputType InputMethod
 		{
 			LeafmostActiveNode.Pin()->CacheFocusRestorationTarget();
 		}
+
+		if (InputMethod == ECommonInputType::Gamepad && CVarSyncCursorToFocusedWidgetOnInputMethodChange.GetValueOnAnyThread())
+		{
+			const int32 OwnerSlateId = GetOwnerUserIndex();
+
+ 			TSharedPtr<FSlateUser> SlateUser = FSlateApplication::Get().GetUser(OwnerSlateId);
+
+			if (SlateUser)
+			{
+				TSharedPtr<SWidget> FocusedWidget = FSlateApplication::Get().GetUserFocusedWidget(OwnerSlateId);
+
+				// ApplyLeafmostNodeConfig may update the currently focused widget, but it doesn't change the cursor position.
+				// It means that next time the user presses the confirmation button on the controller, the mouse click event
+				// will be sent to the last touched position, and not the focused widget. It may make us lose focus and semi-lock inputs.
+				// To avoid this, we update the cursor position to match the focused widget.
+				if (FocusedWidget)
+				{
+					const FGeometry& Geometry = FocusedWidget->GetCachedGeometry();
+					const FVector2D AbsoluteWidgetCenter = Geometry.GetAbsolutePositionAtCoordinates(FVector2D(0.5f, 0.5f));
+
+					SlateUser->SetCursorPosition(AbsoluteWidgetCenter);
+				}
+			}
+		}
 	}
 }
 
