VerseUI: use a proxy widget in primary content and move the actual verse UI to a scene

diff --git a/Engine/Plugins/Runtime/CommonUI/Source/CommonUI/Private/CommonUserWidget.cpp b/Engine/Plugins/Runtime/CommonUI/Source/CommonUI/Private/CommonUserWidget.cpp
--- a/Engine/Plugins/Runtime/CommonUI/Source/CommonUI/Private/CommonUserWidget.cpp
+++ b/Engine/Plugins/Runtime/CommonUI/Source/CommonUI/Private/CommonUserWidget.cpp
@@ -115,7 +115,7 @@ void UCommonUserWidget::AddActionBinding(FUIActionBindingHandle ActionBinding)
 	}
 }
 
-void UCommonUserWidget::RegisterScrollRecipient(const UWidget& AnalogScrollRecipient)
+void UCommonUserWidget::RegisterScrollRecipient(const UWidget& AnalogScrollRecipient, ECommonUIScrollRecipientOwningNodeSource OwningNodeSource)
 {
 	if (!ScrollRecipients.Contains(&AnalogScrollRecipient))
 	{
@@ -124,19 +124,19 @@ void UCommonUserWidget::RegisterScrollRecipient(const UWidget& AnalogScrollRecip
 		{
 			if (UCommonUIActionRouterBase* ActionRouter = UCommonUIActionRouterBase::Get(*this))
 			{
-				ActionRouter->RegisterScrollRecipient(AnalogScrollRecipient);
+				ActionRouter->RegisterScrollRecipient(AnalogScrollRecipient, OwningNodeSource == ECommonUIScrollRecipientOwningNodeSource::RegisteringUserWidget ? this : nullptr);
 			}
 		}
 	}
 }
 
-void UCommonUserWidget::UnregisterScrollRecipient(const UWidget& AnalogScrollRecipient)
+void UCommonUserWidget::UnregisterScrollRecipient(const UWidget& AnalogScrollRecipient, ECommonUIScrollRecipientOwningNodeSource OwningNodeSource)
 {
 	if (ScrollRecipients.Remove(&AnalogScrollRecipient) && GetCachedWidget())
 	{
 		if (UCommonUIActionRouterBase* ActionRouter = UCommonUIActionRouterBase::Get(*this))
 		{
-			ActionRouter->UnregisterScrollRecipient(AnalogScrollRecipient);
+			ActionRouter->UnregisterScrollRecipient(AnalogScrollRecipient, OwningNodeSource == ECommonUIScrollRecipientOwningNodeSource::RegisteringUserWidget ? this : nullptr);
 		}
 	}
 }
diff --git a/Engine/Plugins/Runtime/CommonUI/Source/CommonUI/Private/Input/CommonUIActionRouterBase.cpp b/Engine/Plugins/Runtime/CommonUI/Source/CommonUI/Private/Input/CommonUIActionRouterBase.cpp
--- a/Engine/Plugins/Runtime/CommonUI/Source/CommonUI/Private/Input/CommonUIActionRouterBase.cpp
+++ b/Engine/Plugins/Runtime/CommonUI/Source/CommonUI/Private/Input/CommonUIActionRouterBase.cpp
@@ -416,9 +416,16 @@ void UCommonUIActionRouterBase::SetIsActivatableTreeEnabled(bool bInIsTreeEnable
 	}
 }
 
-void UCommonUIActionRouterBase::RegisterScrollRecipient(const UWidget& ScrollableWidget)
+void UCommonUIActionRouterBase::RegisterScrollRecipient(const UWidget& ScrollableWidget, UWidget* OwningNodeWidget)
 {
-	if (FActivatableTreeNodePtr OwnerNode = FindOwningNode(ScrollableWidget))
+	if (OwningNodeWidget)
+	{
+		if (FActivatableTreeNodePtr OwnerNode = FindOwningNode(*OwningNodeWidget))
+		{
+			OwnerNode->AddScrollRecipient(ScrollableWidget);
+		}
+	}
+	else if (FActivatableTreeNodePtr OwnerNode = FindOwningNode(ScrollableWidget))
 	{
 		OwnerNode->AddScrollRecipient(ScrollableWidget);
 	}
@@ -428,9 +435,16 @@ void UCommonUIActionRouterBase::RegisterScrollRecipient(const UWidget& Scrollabl
 	}
 }
 
-void UCommonUIActionRouterBase::UnregisterScrollRecipient(const UWidget& ScrollableWidget)
+void UCommonUIActionRouterBase::UnregisterScrollRecipient(const UWidget& ScrollableWidget, UWidget* OwningNodeWidget)
 {
-	if (FActivatableTreeNodePtr OwnerNode = FindOwningNode(ScrollableWidget))
+	if (OwningNodeWidget)
+	{
+		if (FActivatableTreeNodePtr OwnerNode = FindOwningNode(*OwningNodeWidget))
+		{
+			OwnerNode->RemoveScrollRecipient(ScrollableWidget);
+		}
+	}
+	else if (FActivatableTreeNodePtr OwnerNode = FindOwningNode(ScrollableWidget))
 	{
 		OwnerNode->RemoveScrollRecipient(ScrollableWidget);
 	}
diff --git a/Engine/Plugins/Runtime/CommonUI/Source/CommonUI/Public/CommonUserWidget.h b/Engine/Plugins/Runtime/CommonUI/Source/CommonUI/Public/CommonUserWidget.h
--- a/Engine/Plugins/Runtime/CommonUI/Source/CommonUI/Public/CommonUserWidget.h
+++ b/Engine/Plugins/Runtime/CommonUI/Source/CommonUI/Public/CommonUserWidget.h
@@ -17,6 +17,16 @@ struct FUIActionTag;
 struct FBindUIActionArgs;
 enum class ECommonInputMode : uint8;
 
+/** Enum to represent which common UI tree node takes scroll input */
+enum ECommonUIScrollRecipientOwningNodeSource : uint8
+{
+	/** The scroll recipient owner takes scroll input */
+	ScrollRecipient,
+
+	/** The userwidget that is registering the scroll recipient takes scroll input. It may or may not be the same as scroll recipient owner. */
+	RegisteringUserWidget
+};
+
 UCLASS(MinimalAPI, ClassGroup = UI, meta = (Category = "Common UI", DisableNativeTick))
 class UCommonUserWidget : public UUserWidget
 {
@@ -84,8 +94,8 @@ protected:
 		return *PC;
 	}
 
-	UE_API void RegisterScrollRecipient(const UWidget& AnalogScrollRecipient);
-	UE_API void UnregisterScrollRecipient(const UWidget& AnalogScrollRecipient);
+	UE_API void RegisterScrollRecipient(const UWidget& AnalogScrollRecipient, ECommonUIScrollRecipientOwningNodeSource OwningNodeSource = ECommonUIScrollRecipientOwningNodeSource::ScrollRecipient);
+	UE_API void UnregisterScrollRecipient(const UWidget& AnalogScrollRecipient, ECommonUIScrollRecipientOwningNodeSource OwningNodeSource = ECommonUIScrollRecipientOwningNodeSource::ScrollRecipient);
 
 #if WITH_EDITOR
 	UE_API virtual const FText GetPaletteCategory() override;
diff --git a/Engine/Plugins/Runtime/CommonUI/Source/CommonUI/Public/Input/CommonUIActionRouterBase.h b/Engine/Plugins/Runtime/CommonUI/Source/CommonUI/Public/Input/CommonUIActionRouterBase.h
--- a/Engine/Plugins/Runtime/CommonUI/Source/CommonUI/Public/Input/CommonUIActionRouterBase.h
+++ b/Engine/Plugins/Runtime/CommonUI/Source/CommonUI/Public/Input/CommonUIActionRouterBase.h
@@ -100,8 +100,8 @@ public:
 	DECLARE_EVENT_OneParam(UCommonUIActionRouterBase, FOnActivationMetadataChanged, FActivationMetadata);
 	FOnActivationMetadataChanged& OnActivationMetadataChanged() const { return OnActivationMetadataChangedEvent; }
 
-	UE_API void RegisterScrollRecipient(const UWidget& ScrollableWidget);
-	UE_API void UnregisterScrollRecipient(const UWidget& ScrollableWidget);
+	UE_API void RegisterScrollRecipient(const UWidget& ScrollableWidget, UWidget* OwningNodeWidget = nullptr);
+	UE_API void UnregisterScrollRecipient(const UWidget& ScrollableWidget, UWidget* OwningNodeWidget = nullptr);
 	UE_API TArray<const UWidget*> GatherActiveAnalogScrollRecipients() const;
 
 	UE_API TArray<FUIActionBindingHandle> GatherActiveBindings() const;
