[CommonUI] Add FindParentWidgetImplementingInterface, complimentary utility to FindParentWidgetOfType

diff --git a/Engine/Plugins/Runtime/CommonUI/Source/CommonUI/Private/CommonUILibrary.cpp b/Engine/Plugins/Runtime/CommonUI/Source/CommonUI/Private/CommonUILibrary.cpp
--- a/Engine/Plugins/Runtime/CommonUI/Source/CommonUI/Private/CommonUILibrary.cpp
+++ b/Engine/Plugins/Runtime/CommonUI/Source/CommonUI/Private/CommonUILibrary.cpp
@@ -48,5 +48,46 @@ UWidget* UCommonUILibrary::FindParentWidgetOfType(UWidget* StartingWidget, TSubc
 	return nullptr;
 }
 
+UWidget* UCommonUILibrary::FindParentWidgetImplementingInterface(UWidget* StartingWidget, TSubclassOf<UInterface> Interface)
+{
+	while (StartingWidget)
+	{
+		UWidget* LocalRoot = StartingWidget;
+		UWidget* LocalParent = LocalRoot->GetParent();
+		while (LocalParent)
+		{
+			if (UClass* LocalParentClass = LocalParent->GetClass())
+			{
+				if (LocalParentClass->ImplementsInterface(Interface))
+				{
+					return LocalParent;
+				}
+			}
+			LocalRoot = LocalParent;
+			LocalParent = LocalParent->GetParent();
+		}
+
+		UWidgetTree* WidgetTree = Cast<UWidgetTree>(LocalRoot->GetOuter());
+		if (WidgetTree == nullptr)
+		{
+			break;
+		}
+
+		StartingWidget = Cast<UUserWidget>(WidgetTree->GetOuter());
+		if (StartingWidget)
+		{
+			if (UClass* StartingWidgetClass = StartingWidget->GetClass())
+			{
+				if (StartingWidgetClass->ImplementsInterface(Interface))
+				{
+					return StartingWidget;
+				}
+			}
+		}
+	}
+
+	return nullptr;
+}
+
 #undef LOCTEXT_NAMESPACE
 
diff --git a/Engine/Plugins/Runtime/CommonUI/Source/CommonUI/Public/CommonUILibrary.h b/Engine/Plugins/Runtime/CommonUI/Source/CommonUI/Public/CommonUILibrary.h
--- a/Engine/Plugins/Runtime/CommonUI/Source/CommonUI/Public/CommonUILibrary.h
+++ b/Engine/Plugins/Runtime/CommonUI/Source/CommonUI/Public/CommonUILibrary.h
@@ -22,6 +22,12 @@ public:
 	 */
 	UFUNCTION(BlueprintCallable, Category="Common UI", meta=(DeterminesOutputType=Type))
 	static UE_API UWidget* FindParentWidgetOfType(UWidget* StartingWidget, TSubclassOf<UWidget> Type);
+
+	/**
+	* Finds the first parent widget implementing the given interface and returns it, or null if no parent could be found.
+	*/
+	UFUNCTION(BlueprintCallable, Category="Common UI")
+	static UE_API UWidget* FindParentWidgetImplementingInterface(UWidget* StartingWidget, TSubclassOf<UInterface> Interface);
 };
 
 #undef UE_API
