### What
- UCommonButtonGroupBase: Exposing bAllowSound input argument through SelectNextButton() and SelectPreviousButton().
### Why
- This allows bAllowSound to be passed to SelectButtonAtIndex() instead of assuming True, which means calls to SelectNextButton() and SelectPreviousButton() can happen without automatically playing the Button Selected sound. Ultimately, this will allow exposure to Blueprints as well, giving designers more control of the UI Audio.
### How
- Each function mentioned is updated to also take bAllowSound as a parameter, which defaults to true to maintain current functionality. When each of those call SelectButtonAtIndex(), this same parameter is passed in, instead of the literal True.

#Synced-CL 48526346

diff --git a/Engine/Plugins/Runtime/CommonUI/Source/CommonUI/Private/Groups/CommonButtonGroupBase.cpp b/Engine/Plugins/Runtime/CommonUI/Source/CommonUI/Private/Groups/CommonButtonGroupBase.cpp
--- a/Engine/Plugins/Runtime/CommonUI/Source/CommonUI/Private/Groups/CommonButtonGroupBase.cpp
+++ b/Engine/Plugins/Runtime/CommonUI/Source/CommonUI/Private/Groups/CommonButtonGroupBase.cpp
@@ -54,7 +54,7 @@ void UCommonButtonGroupBase::DeselectAll()
 	}
 }
 
-void UCommonButtonGroupBase::SelectNextButton(bool bAllowWrap /*= true*/)
+void UCommonButtonGroupBase::SelectNextButton(bool bAllowWrap /*= true*/, const bool bAllowSound /*= true*/)
 {
 	const int32 NumButtons = Buttons.Num();
 	int32 NumRecursiveSelections = 0;
@@ -86,7 +86,7 @@ void UCommonButtonGroupBase::SelectNextButton(bool bAllowWrap /*= true*/)
 			UCommonButtonBase* NextButton = GetButtonBaseAtIndex(SelectionIndex);
 			if (NextButton && NextButton->IsInteractionEnabled())
 			{
-				SelectButtonAtIndex(SelectionIndex);
+				SelectButtonAtIndex(SelectionIndex, bAllowSound);
 			}
 			else
 			{
@@ -99,7 +99,7 @@ void UCommonButtonGroupBase::SelectNextButton(bool bAllowWrap /*= true*/)
 	SelectNextButtonRecursive(SelectedButtonIndex, bAllowWrap);
 }
 
-void UCommonButtonGroupBase::SelectPreviousButton(bool bAllowWrap /*= true*/)
+void UCommonButtonGroupBase::SelectPreviousButton(bool bAllowWrap /*= true*/, const bool bAllowSound /*= true*/)
 {
 	const int32 NumButtons = Buttons.Num();
 	int32 NumRecursiveSelections = 0;
@@ -131,7 +131,7 @@ void UCommonButtonGroupBase::SelectPreviousButton(bool bAllowWrap /*= true*/)
 			UCommonButtonBase* PrevButton = GetButtonBaseAtIndex(SelectionIndex);
 			if (PrevButton && PrevButton->IsInteractionEnabled())
 			{
-				SelectButtonAtIndex(SelectionIndex);
+				SelectButtonAtIndex(SelectionIndex, bAllowSound);
 			}
 			else
 			{
diff --git a/Engine/Plugins/Runtime/CommonUI/Source/CommonUI/Public/Groups/CommonButtonGroupBase.h b/Engine/Plugins/Runtime/CommonUI/Source/CommonUI/Public/Groups/CommonButtonGroupBase.h
--- a/Engine/Plugins/Runtime/CommonUI/Source/CommonUI/Public/Groups/CommonButtonGroupBase.h
+++ b/Engine/Plugins/Runtime/CommonUI/Source/CommonUI/Public/Groups/CommonButtonGroupBase.h
@@ -42,16 +42,18 @@ public:
 	/** 
 	 * Selects the next button in the group 
 	 * @param bAllowWrap Whether to wrap to the first button if the last one is currently selected
+	 * @param bAllowSound Whether the selected button should play its click sound
 	 */
 	UFUNCTION(BlueprintCallable, Category = BaseButtonGroup)
-	UE_API void SelectNextButton(bool bAllowWrap = true);
+	UE_API void SelectNextButton(bool bAllowWrap = true, const bool bAllowSound = true);
 
 	/** 
 	 * Selects the previous button in the group 
 	 * @param bAllowWrap Whether to wrap to the first button if the last one is currently selected
+	 * @param bAllowSound Whether the selected button should play its click sound
 	 */
 	UFUNCTION(BlueprintCallable, Category = BaseButtonGroup)
-	UE_API void SelectPreviousButton(bool bAllowWrap = true);
+	UE_API void SelectPreviousButton(bool bAllowWrap = true, const bool bAllowSound = true);
 
 	/**
 	 * Selects a button at a specific index in the group. Clears all selection if given an invalid index.
