Introduce a maximum size for text common ui text scaling

#tests Played with cvars in the editor, validated that only common texts under the threshold were modified. Also validate that scale is applied to common rich text (the scale will be applyed based on default size, whatever the size of a run is).
#rnx

diff --git a/Engine/Plugins/Runtime/CommonUI/Source/CommonUI/Private/CommonRichTextBlock.cpp b/Engine/Plugins/Runtime/CommonUI/Source/CommonUI/Private/CommonRichTextBlock.cpp
--- a/Engine/Plugins/Runtime/CommonUI/Source/CommonUI/Private/CommonRichTextBlock.cpp
+++ b/Engine/Plugins/Runtime/CommonUI/Source/CommonUI/Private/CommonRichTextBlock.cpp
@@ -132,9 +132,12 @@ public:
 			{
 				Scale = Owner->GetMobileTextBlockScale();
 			}
-			else if (CommonUIUtils::UITextGlobalScaling > 0.0f)
+			else if (CommonUIUtils::UISmallTextScaling > 0.0f)
 			{
-				Scale = CommonUIUtils::UITextGlobalScaling;
+				if (TextStyle.Font.Size <= CommonUIUtils::UISmallTextScalingThresholdMax)
+				{
+					Scale = CommonUIUtils::UISmallTextScaling;
+				}
 			}
 			else
 			{
@@ -607,9 +610,12 @@ void UCommonRichTextBlock::ApplyTextBlockScale() const
 		{
 			Scale = GetMobileTextBlockScale();
 		}
-		else if (CommonUIUtils::UITextGlobalScaling > 0.0f)
+		else if (CommonUIUtils::UISmallTextScaling > 0.0f)
 		{
-			Scale = CommonUIUtils::UITextGlobalScaling;
+			if (GetCurrentDefaultTextStyle().Font.Size <= CommonUIUtils::UISmallTextScalingThresholdMax)
+			{
+				Scale = CommonUIUtils::UISmallTextScaling;
+			}
 		}
 		else
 		{
diff --git a/Engine/Plugins/Runtime/CommonUI/Source/CommonUI/Private/CommonTextBlock.cpp b/Engine/Plugins/Runtime/CommonUI/Source/CommonUI/Private/CommonTextBlock.cpp
--- a/Engine/Plugins/Runtime/CommonUI/Source/CommonUI/Private/CommonTextBlock.cpp
+++ b/Engine/Plugins/Runtime/CommonUI/Source/CommonUI/Private/CommonTextBlock.cpp
@@ -638,9 +638,12 @@ void UCommonTextBlock::ApplyFontSizeMultiplier() const
 		{
 			FontInfo.Size *= GetMobileFontSizeMultiplier();
 		}
-		else if (CommonUIUtils::UITextGlobalScaling > 0.0f)
+		else if (CommonUIUtils::UISmallTextScaling > 0.0f)
 		{
-			FontInfo.Size *= CommonUIUtils::UITextGlobalScaling;
+			if (FontInfo.Size <= CommonUIUtils::UISmallTextScalingThresholdMax)
+			{
+				FontInfo.Size *= CommonUIUtils::UISmallTextScaling;
+			}
 		}
 		else
 		{
diff --git a/Engine/Plugins/Runtime/CommonUI/Source/CommonUI/Private/CommonUIUtils.cpp b/Engine/Plugins/Runtime/CommonUI/Source/CommonUI/Private/CommonUIUtils.cpp
--- a/Engine/Plugins/Runtime/CommonUI/Source/CommonUI/Private/CommonUIUtils.cpp
+++ b/Engine/Plugins/Runtime/CommonUI/Source/CommonUI/Private/CommonUIUtils.cpp
@@ -14,11 +14,17 @@ namespace CommonUIUtils
 		TEXT("Enables Mobile UI Text Scaling"),
 		ECVF_Default);
 
-	float UITextGlobalScaling = -1.0f;
-	FAutoConsoleVariableRef CVarUITextGlobalScaling(
-		TEXT("CommonUI.TextGlobalScaling"),
-		UITextGlobalScaling,
-		TEXT("Scaling to apply to all common ui texts (unless mobile UI text scaling is on). Ignored if 0 or lesser"));
+	float UISmallTextScaling = -1.0f;
+	FAutoConsoleVariableRef CVarUITextSmallTextScaling(
+		TEXT("CommonUI.SmallTextScaling"),
+		UISmallTextScaling,
+		TEXT("Scaling to apply to common ui texts smaller than the configured threshold (unless mobile UI text scaling is on). Ignored if 0 or lesser"));
+
+	float UISmallTextScalingThresholdMax = 15.0f;
+	FAutoConsoleVariableRef CVarUITextSmallTextScalingThresholdMax(
+		TEXT("CommonUI.SmallTextScalingThresholdMax"),
+		UISmallTextScalingThresholdMax,
+		TEXT("Maximum size on which small text Scaling will be applied."));
 
 	bool ShouldDisplayMobileUISizes()
 	{
@@ -30,7 +36,7 @@ namespace CommonUIUtils
 
 	bool ShouldScaleText()
 	{
-		return ShouldDisplayMobileUISizes() || UITextGlobalScaling > 0.0f;
+		return ShouldDisplayMobileUISizes() || UISmallTextScaling > 0.0f;
 	}
 
 	FString PrintAllOwningUserWidgets(const UWidget* Widget)
diff --git a/Engine/Plugins/Runtime/CommonUI/Source/CommonUI/Public/CommonUIUtils.h b/Engine/Plugins/Runtime/CommonUI/Source/CommonUI/Public/CommonUIUtils.h
--- a/Engine/Plugins/Runtime/CommonUI/Source/CommonUI/Public/CommonUIUtils.h
+++ b/Engine/Plugins/Runtime/CommonUI/Source/CommonUI/Public/CommonUIUtils.h
@@ -11,7 +11,8 @@ class IWidgetCompilerLog;
 
 namespace CommonUIUtils
 {
-	extern float UITextGlobalScaling;
+	extern float UISmallTextScaling;
+	extern float UISmallTextScalingThresholdMax;
 
 	enum class ECollisionPolicy
 	{
