CommonUI: Fix not properly resolving the new virtual keys

Since adding Virtual_Gamepad_Accept and Virtual_Gamepad_Back, logic comparing FKeys needs to check the resolved virtual keys (using GetVirtualKey() or IsSameResolvedKey()).

#Synced-CL 48847502

diff --git a/Engine/Plugins/Runtime/CommonUI/Source/CommonInput/Private/CommonInputBaseTypes.cpp b/Engine/Plugins/Runtime/CommonUI/Source/CommonInput/Private/CommonInputBaseTypes.cpp
--- a/Engine/Plugins/Runtime/CommonUI/Source/CommonInput/Private/CommonInputBaseTypes.cpp
+++ b/Engine/Plugins/Runtime/CommonUI/Source/CommonInput/Private/CommonInputBaseTypes.cpp
@@ -31,7 +31,7 @@ namespace CommonUIUtils
 	{
 		const FCommonInputKeyBrushConfiguration* DisplayConfig = InInputBrushDataMap.FindByPredicate([&InKey](const FCommonInputKeyBrushConfiguration& KeyBrushPair) -> bool
 		{
-			return KeyBrushPair.Key == InKey;
+			return KeyBrushPair.Key.IsSameResolvedKey(InKey);
 		});
 
 		if (DisplayConfig)
@@ -57,7 +57,11 @@ namespace CommonUIUtils
 			{
 				for (const FKey& Key : InKeys)
 				{
-					if (!KeyBrushPair.Keys.Contains(Key))
+					const bool bContainsKey = KeyBrushPair.Keys.ContainsByPredicate([Key](const FKey& BrushKey)
+					{
+						return BrushKey.IsSameResolvedKey(Key);
+					});
+					if (!bContainsKey)
 					{
 						return false;
 					}
diff --git a/Engine/Plugins/Runtime/CommonUI/Source/CommonUI/Private/Input/UIActionRouterTypes.cpp b/Engine/Plugins/Runtime/CommonUI/Source/CommonUI/Private/Input/UIActionRouterTypes.cpp
--- a/Engine/Plugins/Runtime/CommonUI/Source/CommonUI/Private/Input/UIActionRouterTypes.cpp
+++ b/Engine/Plugins/Runtime/CommonUI/Source/CommonUI/Private/Input/UIActionRouterTypes.cpp
@@ -692,7 +692,7 @@ EProcessHoldActionResult FActionRouterBindingCollection::ProcessHoldInput(ECommo
 
 					// A persistent displayed action skips the normal rules for reachability, since it'll always appear in a bound action bar
 					const bool bIsDisplayedPersistentAction = Binding->bIsPersistent && Binding->bDisplayInActionBar;
-					if (HoldMapping.Key == Key && (bIsDisplayedPersistentAction || IsWidgetReachableForInput(Binding->BoundWidget.Get())))
+					if (HoldMapping.Key.IsSameResolvedKey(Key) && (bIsDisplayedPersistentAction || IsWidgetReachableForInput(Binding->BoundWidget.Get())))
 					{
 						if (InputEvent == IE_Pressed)
 						{
@@ -761,7 +761,7 @@ bool FActionRouterBindingCollection::ProcessNormalInput(ECommonInputMode ActiveI
 				{
 					// A persistent displayed action skips the normal rules for reachability, since it'll always appear in a bound action bar
 					const bool bIsDisplayedPersistentAction = Binding->bIsPersistent && Binding->bDisplayInActionBar;
-					if (InKey == Key && Binding->InputEvent == InputEvent && (bIsDisplayedPersistentAction || IsWidgetReachableForInput(Binding->BoundWidget.Get())))
+					if (InKey.IsSameResolvedKey(Key) && Binding->InputEvent == InputEvent && (bIsDisplayedPersistentAction || IsWidgetReachableForInput(Binding->BoundWidget.Get())))
 					{
 						// Just in case this was in the middle of a hold process with a different key, reset now
 						Binding->CancelHold();
