[UMG] Exposing a setter to HideInputActionWithKeyboard

- Added a BP setter method for the boolean bHideInputActionWithKeyboard, so TDs can use it in runtime.

#tests In Editor

#Synced-CL 48945164

diff --git a/Engine/Plugins/Runtime/CommonUI/Source/CommonUI/Private/CommonButtonBase.cpp b/Engine/Plugins/Runtime/CommonUI/Source/CommonUI/Private/CommonButtonBase.cpp
--- a/Engine/Plugins/Runtime/CommonUI/Source/CommonUI/Private/CommonButtonBase.cpp
+++ b/Engine/Plugins/Runtime/CommonUI/Source/CommonUI/Private/CommonButtonBase.cpp
@@ -846,6 +846,13 @@ void UCommonButtonBase::SetHideInputAction(bool bInHideInputAction)
 	UpdateInputActionWidgetVisibility();
 }
 
+void UCommonButtonBase::SetHideInputActionWithKeyboard(bool bInHideInputActionWithKeyboard)
+{
+	bHideInputActionWithKeyboard = bInHideInputActionWithKeyboard;
+
+	UpdateInputActionWidgetVisibility();
+}
+
 bool UCommonButtonBase::IsInteractionEnabled() const
 {
 	ESlateVisibility Vis = GetVisibility(); // hidden or collapsed should have 'bInteractionEnabled' set false, but sometimes they don't :(
diff --git a/Engine/Plugins/Runtime/CommonUI/Source/CommonUI/Public/CommonButtonBase.h b/Engine/Plugins/Runtime/CommonUI/Source/CommonUI/Public/CommonButtonBase.h
--- a/Engine/Plugins/Runtime/CommonUI/Source/CommonUI/Public/CommonButtonBase.h
+++ b/Engine/Plugins/Runtime/CommonUI/Source/CommonUI/Public/CommonButtonBase.h
@@ -352,6 +352,9 @@ public:
 	UFUNCTION(BlueprintCallable, Category = "Common Button|Setters")
 	COMMONUI_API void SetHideInputAction(bool bInHideInputAction);
 
+	UFUNCTION(BlueprintCallable, Category = "Common Button|Setters")
+	COMMONUI_API void SetHideInputActionWithKeyboard(bool bInHideInputActionWithKeyboard);
+
 	/** Is this button currently interactable? (use instead of GetIsEnabled) */
 	UFUNCTION(BlueprintCallable, Category = "Common Button|Getters")
 	COMMONUI_API bool IsInteractionEnabled() const;
