Added BlueprintAutocast functions so CommonUI gameplay tag subclasses will default to FGameplayTag comparison operator in blueprint, rather than not being comparable in blueprint.
#Synced-CL 49589669

diff --git a/Engine/Plugins/Runtime/CommonUI/Source/CommonUI/Private/CommonUILibrary.cpp b/Engine/Plugins/Runtime/CommonUI/Source/CommonUI/Private/CommonUILibrary.cpp
--- a/Engine/Plugins/Runtime/CommonUI/Source/CommonUI/Private/CommonUILibrary.cpp
+++ b/Engine/Plugins/Runtime/CommonUI/Source/CommonUI/Private/CommonUILibrary.cpp
@@ -89,5 +89,15 @@ UWidget* UCommonUILibrary::FindParentWidgetImplementingInterface(UWidget* Starti
 	return nullptr;
 }
 
+FGameplayTag UCommonUILibrary::Conv_UITagToGameplayTag(FUITag InValue)
+{
+	return InValue;
+}
+
+FGameplayTag UCommonUILibrary::Conv_UIActionTagToGameplayTag(FUIActionTag InValue)
+{
+	return InValue;
+}
+
 #undef LOCTEXT_NAMESPACE
 
diff --git a/Engine/Plugins/Runtime/CommonUI/Source/CommonUI/Public/CommonUILibrary.h b/Engine/Plugins/Runtime/CommonUI/Source/CommonUI/Public/CommonUILibrary.h
--- a/Engine/Plugins/Runtime/CommonUI/Source/CommonUI/Public/CommonUILibrary.h
+++ b/Engine/Plugins/Runtime/CommonUI/Source/CommonUI/Public/CommonUILibrary.h
@@ -3,6 +3,7 @@
 #pragma once
 
 #include "Kismet/BlueprintFunctionLibrary.h"
+#include "UITag.h"
 #include "CommonUILibrary.generated.h"
 
 #define UE_API COMMONUI_API
@@ -28,6 +29,18 @@ public:
 	*/
 	UFUNCTION(BlueprintCallable, Category="Common UI")
 	static UE_API UWidget* FindParentWidgetImplementingInterface(UWidget* StartingWidget, TSubclassOf<UInterface> Interface);
+
+	/**
+	 * Interpret a UITag as GameplayTag. Enables connecting UITags to wildcard comparison nodes that map to explicit GameplayTag functions.
+	 */
+	UFUNCTION(BlueprintPure, BlueprintInternalUseOnly, meta = (BlueprintAutocast))
+	static UE_API FGameplayTag Conv_UITagToGameplayTag(FUITag InValue);
+
+	/**
+	 * Interpret a UIActionTag as GameplayTag. Enables connecting FUIActionTags to wildcard comparison nodes that map to explicit GameplayTag functions.
+	 */
+	UFUNCTION(BlueprintPure, BlueprintInternalUseOnly, meta = (BlueprintAutocast))
+	static UE_API FGameplayTag Conv_UIActionTagToGameplayTag(FUIActionTag InValue);
 };
 
 #undef UE_API
