When CancelStreaming in UCommonLazyWidget, it has to cancel the async load request also together with hiding the spinner.

#tests Create a popup Modal widget with a button, which calls "SetLazyContext" with any widget followed by calling "SetLazyContext" with null widget when onClicked event. The callback of "SetLazyContext" is print debug string on screen to check the callback being called or not in order to indicate "CancelStreaming" for cancelling async load.
#rnx

#Synced-CL 49675487

diff --git a/Engine/Plugins/Runtime/CommonUI/Source/CommonUI/Private/CommonLazyWidget.cpp b/Engine/Plugins/Runtime/CommonUI/Source/CommonUI/Private/CommonLazyWidget.cpp
--- a/Engine/Plugins/Runtime/CommonUI/Source/CommonUI/Private/CommonLazyWidget.cpp
+++ b/Engine/Plugins/Runtime/CommonUI/Source/CommonUI/Private/CommonLazyWidget.cpp
@@ -88,6 +88,12 @@ void UCommonLazyWidget::SetForceShowSpinner(bool bShowLoading)
 
 void UCommonLazyWidget::CancelStreaming()
 {
+	if (StreamingHandle.IsValid())
+	{
+		StreamingHandle->CancelHandle();
+		StreamingHandle.Reset();
+	}
+
 	if (MyLoadGuard.IsValid())
 	{
 		MyLoadGuard->SetForceShowSpinner(false);
