[CommonInputSubsystem] Add input type lock when using a virtual keyboard

#rnx

#Synced-CL 49981045

diff --git a/Engine/Plugins/Runtime/CommonUI/Source/CommonInput/Private/CommonInputSubsystem.cpp b/Engine/Plugins/Runtime/CommonUI/Source/CommonInput/Private/CommonInputSubsystem.cpp
--- a/Engine/Plugins/Runtime/CommonUI/Source/CommonInput/Private/CommonInputSubsystem.cpp
+++ b/Engine/Plugins/Runtime/CommonUI/Source/CommonInput/Private/CommonInputSubsystem.cpp
@@ -14,6 +14,7 @@
 #include "Framework/Application/IInputProcessor.h"
 #include "Framework/Application/SlateApplication.h"
 #include "Framework/Application/SlateUser.h"
+#include "GenericPlatform/GenericApplication.h"
 #include "Misc/ConfigCacheIni.h"
 #include "HAL/PlatformStackWalk.h"
 #include "Widgets/SViewport.h"
@@ -91,6 +92,22 @@ void UCommonInputSubsystem::Initialize(FSubsystemCollectionBase& Collection)
 	CVarInputKeysVisible->SetOnChangedCallback(FConsoleVariableDelegate::CreateUObject(this, &UCommonInputSubsystem::ShouldShowInputKeysChanged));
 
 	SetActionDomainTable(FCommonInputBase::GetInputSettings()->GetActionDomainTable());
+
+	if (FSlateApplication::IsInitialized())
+	{
+		if (GenericApplication* GenericApplication = FSlateApplication::Get().GetPlatformApplication().Get())
+		{
+			if (!GenericApplication->OnVirtualKeyboardShown().IsBoundToObject(this))
+			{
+				GenericApplication->OnVirtualKeyboardShown().AddUObject(this, &ThisClass::HandleOnVirtualKeyboardShown);
+			}
+
+			if (!GenericApplication->OnVirtualKeyboardHidden().IsBoundToObject(this))
+			{
+				GenericApplication->OnVirtualKeyboardHidden().AddUObject(this, &ThisClass::HandleOnVirtualKeyboardHidden);
+			}
+		}
+	}
 }
 
 void UCommonInputSubsystem::Deinitialize()
@@ -103,6 +120,15 @@ void UCommonInputSubsystem::Deinitialize()
 	CommonInputPreprocessor.Reset();
 
 	FTSTicker::RemoveTicker(TickHandle);
+
+	if (FSlateApplication::IsInitialized())
+	{
+		if (GenericApplication* GenericApplication = FSlateApplication::Get().GetPlatformApplication().Get())
+		{
+			GenericApplication->OnVirtualKeyboardShown().RemoveAll(this);
+			GenericApplication->OnVirtualKeyboardHidden().RemoveAll(this);
+		}
+	}
 }
 
 FGamepadChangeDetectedEvent& UCommonInputSubsystem::GetOnGamepadChangeDetected()
@@ -309,6 +335,22 @@ void UCommonInputSubsystem::RecalculateCurrentInputType()
 	}
 }
 
+void UCommonInputSubsystem::HandleOnVirtualKeyboardShown(FPlatformRect VkRect)
+{
+	if (FPlatformMisc::IsDeviceGamingHandheld() && CurrentInputType == ECommonInputType::Gamepad)
+	{
+		AddOrRemoveInputTypeLock(TEXT("HandheldVirtualKeyboard"), ECommonInputType::Gamepad, true);
+	}
+}
+
+void UCommonInputSubsystem::HandleOnVirtualKeyboardHidden()
+{
+	if (FPlatformMisc::IsDeviceGamingHandheld())
+	{
+		AddOrRemoveInputTypeLock(TEXT("HandheldVirtualKeyboard"), ECommonInputType::Gamepad, false);
+	}
+}
+
 void UCommonInputSubsystem::SetCurrentInputType(ECommonInputType NewInputType)
 {
 	if (((RawInputType != NewInputType) || bInputMethodLockedByThrashing) && PlatformSupportsInputType(NewInputType))
diff --git a/Engine/Plugins/Runtime/CommonUI/Source/CommonInput/Public/CommonInputSubsystem.h b/Engine/Plugins/Runtime/CommonUI/Source/CommonInput/Public/CommonInputSubsystem.h
--- a/Engine/Plugins/Runtime/CommonUI/Source/CommonInput/Public/CommonInputSubsystem.h
+++ b/Engine/Plugins/Runtime/CommonUI/Source/CommonInput/Public/CommonInputSubsystem.h
@@ -15,6 +15,7 @@ class APlayerController;
 class FCommonInputPreprocessor;
 class FSlateUser;
 class UCommonInputActionDomainTable;
+struct FPlatformRect;
 
 DECLARE_DYNAMIC_MULTICAST_DELEGATE_OneParam(FInputMethodChangedDelegate, ECommonInputType, bNewInputType);
 DECLARE_MULTICAST_DELEGATE_ThreeParams(FPlatformInputSupportOverrideDelegate, ULocalPlayer*, ECommonInputType, bool&);
@@ -129,6 +130,9 @@ private:
 
 	UE_API void RecalculateCurrentInputType();
 
+	void HandleOnVirtualKeyboardShown(FPlatformRect VkRect);
+	void HandleOnVirtualKeyboardHidden();
+
 	FTSTicker::FDelegateHandle TickHandle;
 
 	UPROPERTY(BlueprintAssignable, Category = CommonInputSubsystem, meta = (AllowPrivateAccess))
