Changed the SetMinDimension and RefreshDimension to not override variable from the button.
#rnx
#Synced-CL 50451663

diff --git a/Engine/Plugins/Runtime/CommonUI/Source/CommonUI/Private/CommonButtonBase.cpp b/Engine/Plugins/Runtime/CommonUI/Source/CommonUI/Private/CommonButtonBase.cpp
--- a/Engine/Plugins/Runtime/CommonUI/Source/CommonUI/Private/CommonButtonBase.cpp
+++ b/Engine/Plugins/Runtime/CommonUI/Source/CommonUI/Private/CommonButtonBase.cpp
@@ -40,6 +40,13 @@ namespace UE::CommonUI::Private
 		EnableSimulateHoverOnTouchInput,
 		TEXT("Allows buttons to simulate hovering on touch in accordance with the property SimulateHoverOnTouchInput.\n0: Disable, 1: Enable (default), 2: Legacy Mode (Deprecated)\n\nNote on Legacy Mode: This mode was previously the default (CommonButton.EnableSimulateHoverOnTouch=1) and is preserved for retro-compatibility with CommonButton.EnableSimulateHoverOnTouch=2 but will be removed in a future update. In this legacy implementation, the property SimulateHoverOnTouch=true simulates Hover events as expected. However, SimulateHoverOnTouch=false was not blocking all hover events on touch. This implementation was simply preventing the Press & Release functions from simulating more hover events. This implementation was causing inconsistent behaviors for widgets with SimulateHoverOnTouch=false. The new implementation will effectively block all Hover events for UCommonButtonBase if the property SimulateHoverOnTouch is set to false in the editor. This ensures a consistent behavior for touch input so a widget can be built without any Hover events being simulated on a touch screen."),
 		ECVF_Default);
+
+	bool bDisableDimensionValueSetInDesignTime = false;
+	FAutoConsoleVariableRef CVarDisableDimensionValueSetInDesignTime(
+		TEXT("CommonButton.DisableDimensionValueSetInDesignTime"),
+		bDisableDimensionValueSetInDesignTime,
+		TEXT("Block SetMin/MaxDimension to set directly the dimension variable (MinHeight, MaxHeight,MinWidth, MaxWidth) during design time to prevent possible data corruptoion from Mobile value"),
+		ECVF_Default);
 }
 
 //////////////////////////////////////////////////////////////////////////
@@ -1129,36 +1136,41 @@ void UCommonButtonBase::SetSelectedInternal(bool bInSelected, bool bAllowSound /
 }
 
 void UCommonButtonBase::RefreshDimensions()
+{
+	RefreshDimensions(MinWidth, MinHeight, MaxWidth, MaxHeight);
+}
+
+void UCommonButtonBase::RefreshDimensions(const int32 InMinWidth, const int32 InMinHeight, const int32 InMaxWidth, const int32 InMaxHeight)
 {
 	if (RootButton.IsValid())
 	{
 		const UCommonButtonStyle* const StyleCDO = GetStyleCDO();
-		RootButton->SetMinDesiredWidth(FMath::Max(MinWidth, StyleCDO ? StyleCDO->MinWidth : 0));
-		RootButton->SetMinDesiredHeight(FMath::Max(MinHeight, StyleCDO ? StyleCDO->MinHeight : 0));
+		RootButton->SetMinDesiredWidth(FMath::Max(InMinWidth, StyleCDO ? StyleCDO->MinWidth : 0));
+		RootButton->SetMinDesiredHeight(FMath::Max(InMinHeight, StyleCDO ? StyleCDO->MinHeight : 0));
 		
 		if (!StyleCDO)
 		{
-			RootButton->SetMaxDesiredWidth(MaxWidth);
-			RootButton->SetMaxDesiredHeight(MaxHeight);
+			RootButton->SetMaxDesiredWidth(InMaxWidth);
+			RootButton->SetMaxDesiredHeight(InMaxHeight);
 		}
 		else
 		{
-			if (MaxWidth > 0 && StyleCDO->MaxWidth > 0)
+			if (InMaxWidth > 0 && StyleCDO->MaxWidth > 0)
 			{
-				RootButton->SetMaxDesiredWidth(FMath::Min(MaxWidth, StyleCDO->MaxWidth));
+				RootButton->SetMaxDesiredWidth(FMath::Min(InMaxWidth, StyleCDO->MaxWidth));
 			}
 			else
 			{
-				RootButton->SetMaxDesiredWidth(FMath::Max(MaxWidth, StyleCDO->MaxWidth));
+				RootButton->SetMaxDesiredWidth(FMath::Max(InMaxWidth, StyleCDO->MaxWidth));
 			}
 
-			if (MaxHeight > 0 && StyleCDO->MaxHeight > 0)
+			if (InMaxHeight > 0 && StyleCDO->MaxHeight > 0)
 			{
-				RootButton->SetMaxDesiredHeight(FMath::Min(MaxHeight, StyleCDO->MaxHeight));
+				RootButton->SetMaxDesiredHeight(FMath::Min(InMaxHeight, StyleCDO->MaxHeight));
 			}
 			else
 			{
-				RootButton->SetMaxDesiredHeight(FMath::Max(MaxHeight, StyleCDO->MaxHeight));
+				RootButton->SetMaxDesiredHeight(FMath::Max(InMaxHeight, StyleCDO->MaxHeight));
 			}
 		}
 	}
@@ -1314,18 +1326,38 @@ TSubclassOf<UCommonTextStyle> UCommonButtonBase::GetCurrentTextStyleClass() cons
 
 void UCommonButtonBase::SetMinDimensions(int32 InMinWidth, int32 InMinHeight)
 {
-	MinWidth = InMinWidth;
-	MinHeight = InMinHeight;
+	InMinWidth = FMath::Max(InMinWidth, 0);
+	InMinHeight = FMath::Max(InMinHeight, 0);
 
-	RefreshDimensions();
+	if (UE::CommonUI::Private::bDisableDimensionValueSetInDesignTime && IsDesignTime())
+	{
+		RefreshDimensions(InMinWidth, InMinHeight, RootButton->GetMaxWidth(), RootButton->GetMaxHeight());
+	}
+	else
+	{
+		MinWidth = InMinWidth;
+		MinHeight = InMinHeight;
+
+		RefreshDimensions();
+	}
 }
 
 void UCommonButtonBase::SetMaxDimensions(int32 InMaxWidth, int32 InMaxHeight)
 {
-	MaxWidth = InMaxWidth;
-	MaxHeight = InMaxHeight;
-	
-	RefreshDimensions();
+	InMaxWidth = FMath::Max(InMaxWidth, RootButton->GetMinWidth());
+	InMaxHeight = FMath::Max(InMaxHeight, RootButton->GetMinHeight());
+
+	if (UE::CommonUI::Private::bDisableDimensionValueSetInDesignTime && IsDesignTime())
+	{
+		RefreshDimensions(RootButton->GetMinWidth(), RootButton->GetMinHeight(), InMaxWidth, InMaxHeight);
+	}
+	else
+	{
+		MaxWidth = InMaxWidth;
+		MaxHeight = InMaxHeight;
+
+		RefreshDimensions();
+	}
 }
 
 void UCommonButtonBase::SetTriggeredInputAction(const FDataTableRowHandle &InputActionRow)
diff --git a/Engine/Plugins/Runtime/CommonUI/Source/CommonUI/Public/CommonButtonBase.h b/Engine/Plugins/Runtime/CommonUI/Source/CommonUI/Public/CommonButtonBase.h
--- a/Engine/Plugins/Runtime/CommonUI/Source/CommonUI/Public/CommonButtonBase.h
+++ b/Engine/Plugins/Runtime/CommonUI/Source/CommonUI/Public/CommonButtonBase.h
@@ -254,9 +254,13 @@ public:
 	COMMONUI_API bool IsPressed() const;
 
 	COMMONUI_API void SetMinDesiredHeight(int32 InMinHeight);
+	COMMONUI_API int32 GetMinHeight() const { return MinHeight; }
 	COMMONUI_API void SetMinDesiredWidth(int32 InMinWidth);
+	COMMONUI_API int32 GetMinWidth() const { return MinWidth; }
 	COMMONUI_API void SetMaxDesiredHeight(int32 InMaxHeight);
+	COMMONUI_API int32 GetMaxHeight() const { return MaxHeight; }
 	COMMONUI_API void SetMaxDesiredWidth(int32 InMaxWidth);
+	COMMONUI_API int32 GetMaxWidth() const { return MaxWidth; }
 
 	inline TSharedPtr<class SCommonButton> GetCommonButton() const { return MyCommonButton; }
 
@@ -791,6 +795,7 @@ protected:
 	COMMONUI_API virtual bool GetButtonAnalyticInfo(FString& ButtonName, FString& ABTestName, FString& ExtraData) const;
 
 	COMMONUI_API void RefreshDimensions();
+	COMMONUI_API void RefreshDimensions(const int32 InMinWidth,const int32 InMinHeight,const int32 InMaxWidth,const int32 InMaxHeight);
 	COMMONUI_API virtual void NativeOnMouseEnter( const FGeometry& InGeometry, const FPointerEvent& InMouseEvent ) override;
 	COMMONUI_API virtual void NativeOnMouseLeave( const FPointerEvent& InMouseEvent ) override;
 
