CommonUIActionRouter: Cleaned up handling of SetIgnoreMoveInput/SetIgnoreLookInput, as it was in some cases being decremented without first being incremented

#Synced-CL 50417652

diff --git a/Engine/Plugins/Runtime/CommonUI/Source/CommonUI/Private/Input/CommonUIActionRouterBase.cpp b/Engine/Plugins/Runtime/CommonUI/Source/CommonUI/Private/Input/CommonUIActionRouterBase.cpp
--- a/Engine/Plugins/Runtime/CommonUI/Source/CommonUI/Private/Input/CommonUIActionRouterBase.cpp
+++ b/Engine/Plugins/Runtime/CommonUI/Source/CommonUI/Private/Input/CommonUIActionRouterBase.cpp
@@ -885,7 +885,27 @@ void UCommonUIActionRouterBase::HandleRootWidgetSlateReleased(TWeakPtr<FActivata
 
 		if (RootNodes.Num() == 0)
 		{
-			ActiveInputConfig.Reset();
+			// Clear any input locks before clearing our active config so they don't leak
+			if (ActiveInputConfig.IsSet())
+			{
+				if (ULocalPlayer* LocalPlayer = GetLocalPlayer())
+				{
+					if (APlayerController* PC = LocalPlayer->GetPlayerController(GetWorld()))
+					{
+						if (ActiveInputConfig->bIgnoreMoveInput)
+						{
+							PC->SetIgnoreMoveInput(false);
+						}
+
+						if (ActiveInputConfig->bIgnoreLookInput)
+						{
+							PC->SetIgnoreLookInput(false);
+						}
+					}
+				}
+				ActiveInputConfig.Reset();
+			}
+
 			RefreshActionDomainLeafNodeConfig();
 		}
 	}
@@ -1820,12 +1840,15 @@ void UCommonUIActionRouterBase::ApplyUIInputConfig(const FUIInputConfig& NewConf
 			{
 				if (APlayerController* PC = LocalPlayer.GetPlayerController(GetWorld()))
 				{
-					if (!OldConfig.IsSet() || OldConfig.GetValue().bIgnoreMoveInput != NewConfig.bIgnoreMoveInput)
+					// Ensure we don't ever call SetIgnoreMoveInput(false) unless we previously called SetIgnoreMoveInput(true)
+					const bool bOldIgnoreMoveValue = OldConfig.IsSet() && OldConfig.GetValue().bIgnoreMoveInput;
+					if (bOldIgnoreMoveValue != NewConfig.bIgnoreMoveInput)
 					{
 						PC->SetIgnoreMoveInput(NewConfig.bIgnoreMoveInput);						
 					}
-					
-					if (!OldConfig.IsSet() || OldConfig.GetValue().bIgnoreLookInput != NewConfig.bIgnoreLookInput)
+					// Ensure we don't ever call SetIgnoreLookInput(false) unless we previously called SetIgnoreLookInput(true)
+					const bool bOldIgnoreLookValue = OldConfig.IsSet() && OldConfig.GetValue().bIgnoreLookInput;
+					if (bOldIgnoreLookValue != NewConfig.bIgnoreLookInput)
 					{
 						PC->SetIgnoreLookInput(NewConfig.bIgnoreLookInput);						
 					}
