BoundActionBar: Cannot change the DisplayName button text for CommonBoundActionButton at runtime

#tests PIE BP function call dynamically updates DisplayName for buttons

#Synced-CL 50544476

diff --git a/Engine/Plugins/Runtime/CommonUI/Source/CommonUI/Private/CommonButtonBase.cpp b/Engine/Plugins/Runtime/CommonUI/Source/CommonUI/Private/CommonButtonBase.cpp
--- a/Engine/Plugins/Runtime/CommonUI/Source/CommonUI/Private/CommonButtonBase.cpp
+++ b/Engine/Plugins/Runtime/CommonUI/Source/CommonUI/Private/CommonButtonBase.cpp
@@ -1360,6 +1360,19 @@ void UCommonButtonBase::SetMaxDimensions(int32 InMaxWidth, int32 InMaxHeight)
 	}
 }
 
+void UCommonButtonBase::SetActionDisplayName(const FText& InDisplayName)
+{
+	if (!TriggeringBindingHandle.IsValid())
+	{
+		return;
+	}
+
+	if (!TriggeringBindingHandle.GetDisplayName().EqualTo(InDisplayName))
+	{
+		TriggeringBindingHandle.SetDisplayName(InDisplayName);
+	}
+}
+
 void UCommonButtonBase::SetTriggeredInputAction(const FDataTableRowHandle &InputActionRow)
 {
 	UnbindTriggeringInputActionToClick();
diff --git a/Engine/Plugins/Runtime/CommonUI/Source/CommonUI/Public/CommonButtonBase.h b/Engine/Plugins/Runtime/CommonUI/Source/CommonUI/Public/CommonButtonBase.h
--- a/Engine/Plugins/Runtime/CommonUI/Source/CommonUI/Public/CommonButtonBase.h
+++ b/Engine/Plugins/Runtime/CommonUI/Source/CommonUI/Public/CommonButtonBase.h
@@ -459,6 +459,10 @@ public:
 	UFUNCTION(BlueprintCallable, Category = "Common Button|Setters")
 	COMMONUI_API void SetMaxDimensions(int32 InMaxWidth, int32 InMaxHeight);
 
+	/** Updates the display name via its bound UIAction **/
+	UFUNCTION(BlueprintCallable, Category = "Common Button|Setters")
+	COMMONUI_API void SetActionDisplayName(const FText& InDisplayName);
+
 	/** Updates the current triggered action */
 	UFUNCTION(BlueprintCallable, Category = "Common Button|Setters")
 	COMMONUI_API void SetTriggeredInputAction(const FDataTableRowHandle &InputActionRow);
