Resolve Crash from null ptr Root Button in Common Button Base

#Synced-CL 50792027

diff --git a/Engine/Plugins/Runtime/CommonUI/Source/CommonUI/Private/CommonButtonBase.cpp b/Engine/Plugins/Runtime/CommonUI/Source/CommonUI/Private/CommonButtonBase.cpp
--- a/Engine/Plugins/Runtime/CommonUI/Source/CommonUI/Private/CommonButtonBase.cpp
+++ b/Engine/Plugins/Runtime/CommonUI/Source/CommonUI/Private/CommonButtonBase.cpp
@@ -1331,7 +1331,10 @@ void UCommonButtonBase::SetMinDimensions(int32 InMinWidth, int32 InMinHeight)
 
 	if (UE::CommonUI::Private::bDisableDimensionValueSetInDesignTime && IsDesignTime())
 	{
-		RefreshDimensions(InMinWidth, InMinHeight, RootButton->GetMaxWidth(), RootButton->GetMaxHeight());
+		if (RootButton.IsValid())
+		{
+			RefreshDimensions(InMinWidth, InMinHeight, RootButton->GetMaxWidth(), RootButton->GetMaxHeight());
+		}
 	}
 	else
 	{
@@ -1344,12 +1347,15 @@ void UCommonButtonBase::SetMinDimensions(int32 InMinWidth, int32 InMinHeight)
 
 void UCommonButtonBase::SetMaxDimensions(int32 InMaxWidth, int32 InMaxHeight)
 {
-	InMaxWidth = FMath::Max(InMaxWidth, RootButton->GetMinWidth());
-	InMaxHeight = FMath::Max(InMaxHeight, RootButton->GetMinHeight());
+	InMaxWidth = FMath::Max(InMaxWidth, RootButton.IsValid() ? RootButton->GetMinWidth() : 0);
+	InMaxHeight = FMath::Max(InMaxHeight, RootButton.IsValid() ? RootButton->GetMinHeight() : 0);
 
 	if (UE::CommonUI::Private::bDisableDimensionValueSetInDesignTime && IsDesignTime())
 	{
-		RefreshDimensions(RootButton->GetMinWidth(), RootButton->GetMinHeight(), InMaxWidth, InMaxHeight);
+		if (RootButton.IsValid())
+		{
+			RefreshDimensions(RootButton->GetMinWidth(), RootButton->GetMinHeight(), InMaxWidth, InMaxHeight);
+		}
 	}
 	else
 	{
