Better verbose logging for Common Activatable widget

#changelist validated

#Synced-CL 50897367

diff --git a/Engine/Plugins/Runtime/CommonUI/Source/CommonUI/Private/Input/UIActionRouterTypes.cpp b/Engine/Plugins/Runtime/CommonUI/Source/CommonUI/Private/Input/UIActionRouterTypes.cpp
--- a/Engine/Plugins/Runtime/CommonUI/Source/CommonUI/Private/Input/UIActionRouterTypes.cpp
+++ b/Engine/Plugins/Runtime/CommonUI/Source/CommonUI/Private/Input/UIActionRouterTypes.cpp
@@ -1470,6 +1470,8 @@ bool FActivatableTreeRoot::UpdateLeafmostActiveNode(FActivatableTreeNodePtr Base
 
 	if (NewLeafmostNode != LeafmostActiveNode)
 	{
+		TSharedPtr<SWidget> OldLeafmostSlateWidget;
+
 		if (LeafmostActiveNode.IsValid())
 		{
 			TSharedPtr<FActivatableTreeNode> PinnedLeafmostActiveNode = LeafmostActiveNode.Pin();
@@ -1477,6 +1479,7 @@ bool FActivatableTreeRoot::UpdateLeafmostActiveNode(FActivatableTreeNodePtr Base
 
 			if (UCommonActivatableWidget* LeafmostActivatableWidget = PinnedLeafmostActiveNode->GetWidget())
 			{
+				OldLeafmostSlateWidget = LeafmostActivatableWidget->GetCachedWidget();
 				LeafmostActivatableWidget->OnRequestRefreshFocus().RemoveAll(this);
 			}
 		}
@@ -1491,7 +1494,21 @@ bool FActivatableTreeRoot::UpdateLeafmostActiveNode(FActivatableTreeNodePtr Base
 		{
 			if (UCommonActivatableWidget* NewLeafmostWidget = LeafmostActiveNode.Pin()->GetWidget())
 			{
+				const TSharedPtr<SWidget> NewLeafmostSlateWidget = NewLeafmostWidget->GetCachedWidget();
+
 				NewLeafmostWidget->OnRequestRefreshFocus().AddSP(this, &FActivatableTreeRoot::HandleRequestRefreshLeafmostFocus);
+				
+				UE_LOG(
+					LogUIActionRouter,
+					VeryVerbose,
+					TEXT("Leafmost activatable widget has changed from [%s] to [%s]"),
+					OldLeafmostSlateWidget.IsValid()
+					? *FReflectionMetaData::GetWidgetDebugInfo(OldLeafmostSlateWidget.Get())
+					: TEXT("None"),
+					NewLeafmostSlateWidget.IsValid()
+					? *FReflectionMetaData::GetWidgetDebugInfo(NewLeafmostSlateWidget.Get())
+					: TEXT("None")
+				);
 			}
 		}
 
