RequestRefreshFocusIfDescendant - allow a UWidget to request a focus refresh if it is a descendant of the current leaf-most activatable widget.

#changelist validated

#tests Frontend

#Synced-CL 50950564

diff --git a/Engine/Plugins/Runtime/CommonUI/Source/CommonUI/Private/CommonUILibrary.cpp b/Engine/Plugins/Runtime/CommonUI/Source/CommonUI/Private/CommonUILibrary.cpp
--- a/Engine/Plugins/Runtime/CommonUI/Source/CommonUI/Private/CommonUILibrary.cpp
+++ b/Engine/Plugins/Runtime/CommonUI/Source/CommonUI/Private/CommonUILibrary.cpp
@@ -1,8 +1,13 @@
 // Copyright Epic Games, Inc. All Rights Reserved.
 
 #include "CommonUILibrary.h"
-#include "Components/PanelWidget.h"
+
 #include "Blueprint/WidgetTree.h"
+#include "CommonActivatableWidget.h"
+#include "CommonUIPrivate.h"
+#include "Components/PanelWidget.h"
+#include "Input/CommonUIActionRouterBase.h"
+#include "Input/UIActionRouterTypes.h"
 
 #include UE_INLINE_GENERATED_CPP_BY_NAME(CommonUILibrary)
 
@@ -16,6 +21,31 @@ UCommonUILibrary::UCommonUILibrary(const FObjectInitializer& ObjectInitializer)
 {
 }
 
+bool UCommonUILibrary::RequestRefreshFocusIfLeafmostDescendant(UWidget* ContextWidget)
+{
+	if (!ContextWidget)
+	{
+		UE_LOG(LogCommonUI, Warning, TEXT("RequestRefreshFocusIfLeafmostDescendant was called on/from a null widget."));
+		return false;
+	}
+	
+	if(UCommonUIActionRouterBase* ActionRouter = UCommonUIActionRouterBase::Get(*ContextWidget))
+	{
+		// Ensure that the widget is actually hosted by the current leaf-most activatable node
+		if (UCommonActivatableWidget* ActivatableWidget = ActionRouter->GetLeafmostActivatableWidget())
+		{
+			// Ensure that the widget is actually hosted and that it is a descendant of the leaf-most activatable widget
+			if (ActionRouter->IsWidgetInLeafmostNodeHierarchy(*ContextWidget))
+			{
+				ActivatableWidget->RequestRefreshFocus();
+				return true;
+			}
+		}
+	}
+	
+	return false;
+}
+
 UWidget* UCommonUILibrary::FindParentWidgetOfType(UWidget* StartingWidget, TSubclassOf<UWidget> Type)
 {
 	while ( StartingWidget )
diff --git a/Engine/Plugins/Runtime/CommonUI/Source/CommonUI/Private/Input/CommonUIActionRouterBase.cpp b/Engine/Plugins/Runtime/CommonUI/Source/CommonUI/Private/Input/CommonUIActionRouterBase.cpp
--- a/Engine/Plugins/Runtime/CommonUI/Source/CommonUI/Private/Input/CommonUIActionRouterBase.cpp
+++ b/Engine/Plugins/Runtime/CommonUI/Source/CommonUI/Private/Input/CommonUIActionRouterBase.cpp
@@ -749,6 +749,28 @@ bool UCommonUIActionRouterBase::IsWidgetInActiveRoot(const UCommonActivatableWid
 	return false;
 }
 
+bool UCommonUIActionRouterBase::IsWidgetInLeafmostNodeHierarchy(const UWidget& Widget) const
+{
+	FActivatableTreeRootPtr RootNode = ActiveRootNode ? ActiveRootNode : FindActiveActionDomainRootNode();
+	if (!RootNode)
+	{
+		return false;
+	}
+
+	return RootNode->IsWidgetInLeafmostNodeHierarchy(Widget);
+}
+
+UCommonActivatableWidget* UCommonUIActionRouterBase::GetLeafmostActivatableWidget() const
+{
+	FActivatableTreeRootPtr RootNode = ActiveRootNode ? ActiveRootNode : FindActiveActionDomainRootNode();
+	if (!RootNode)
+	{
+		return nullptr;
+	}
+
+	return RootNode->GetLeafmostActivatableWidget();
+}
+
 void UCommonUIActionRouterBase::NotifyUserWidgetConstructed(const UCommonUserWidget& Widget)
 {
 	check(Widget.GetCachedWidget().IsValid());
diff --git a/Engine/Plugins/Runtime/CommonUI/Source/CommonUI/Private/Input/UIActionRouterTypes.cpp b/Engine/Plugins/Runtime/CommonUI/Source/CommonUI/Private/Input/UIActionRouterTypes.cpp
--- a/Engine/Plugins/Runtime/CommonUI/Source/CommonUI/Private/Input/UIActionRouterTypes.cpp
+++ b/Engine/Plugins/Runtime/CommonUI/Source/CommonUI/Private/Input/UIActionRouterTypes.cpp
@@ -1649,6 +1649,30 @@ bool FActivatableTreeRoot::CanSetInputConfigAndFocus() const
 	return !IsAnActionDomainRoot() || IsActiveActionDomainRoot();
 }
 
+bool FActivatableTreeRoot::IsWidgetInLeafmostNodeHierarchy(const UWidget& Widget) const
+{
+	if (FActivatableTreeNodePtr PinnedLeaf = LeafmostActiveNode.Pin())
+	{
+		// If your leaf node has an SWidget-level helper, use it; otherwise adapt to UWidget
+		TSharedPtr<SWidget> Cached = Widget.GetCachedWidget();
+		if (Cached.IsValid())
+		{
+			return PinnedLeaf->IsExclusiveParentOfWidget(Cached);
+		}
+	}
+	return false;
+}
+
+UCommonActivatableWidget* FActivatableTreeRoot::GetLeafmostActivatableWidget() const
+{
+	if (FActivatableTreeNodePtr PinnedLeaf = LeafmostActiveNode.Pin())
+	{
+		return PinnedLeaf->GetWidget();
+	}
+
+	return nullptr;
+}
+
 void FActivatableTreeRoot::FocusLeafmostNode()
 {
 	if (!CanSetInputConfigAndFocus())
diff --git a/Engine/Plugins/Runtime/CommonUI/Source/CommonUI/Private/Input/UIActionRouterTypes.h b/Engine/Plugins/Runtime/CommonUI/Source/CommonUI/Private/Input/UIActionRouterTypes.h
--- a/Engine/Plugins/Runtime/CommonUI/Source/CommonUI/Private/Input/UIActionRouterTypes.h
+++ b/Engine/Plugins/Runtime/CommonUI/Source/CommonUI/Private/Input/UIActionRouterTypes.h
@@ -210,6 +210,9 @@ public:
 	UE_API bool IsAnActionDomainRoot() const;
 	UE_API bool IsActiveActionDomainRoot() const;
 	UE_API bool CanSetInputConfigAndFocus() const;
+	UE_API bool IsWidgetInLeafmostNodeHierarchy(const UWidget& Widget) const;
+
+	UE_API UCommonActivatableWidget* GetLeafmostActivatableWidget() const;
 
 protected:
 	UE_API virtual void Init() override;
diff --git a/Engine/Plugins/Runtime/CommonUI/Source/CommonUI/Public/CommonUILibrary.h b/Engine/Plugins/Runtime/CommonUI/Source/CommonUI/Public/CommonUILibrary.h
--- a/Engine/Plugins/Runtime/CommonUI/Source/CommonUI/Public/CommonUILibrary.h
+++ b/Engine/Plugins/Runtime/CommonUI/Source/CommonUI/Public/CommonUILibrary.h
@@ -18,6 +18,12 @@ class UCommonUILibrary : public UBlueprintFunctionLibrary
 
 public:
 
+	/**
+	 * Allow any widget that is a descendant of the leaf-most widget to initiate a refresh focus request
+	 */
+	UFUNCTION(BlueprintCallable, Category = "CommonUI|Focus", meta=(DefaultToSelf="ContextWidget", WorldContext="ContextWidget"))
+	static UE_API bool RequestRefreshFocusIfLeafmostDescendant(UWidget* ContextWidget);
+
 	/**
 	 * Finds the first parent widget of the given type and returns it, or null if no parent could be found.
 	 */
diff --git a/Engine/Plugins/Runtime/CommonUI/Source/CommonUI/Public/Input/CommonUIActionRouterBase.h b/Engine/Plugins/Runtime/CommonUI/Source/CommonUI/Public/Input/CommonUIActionRouterBase.h
--- a/Engine/Plugins/Runtime/CommonUI/Source/CommonUI/Public/Input/CommonUIActionRouterBase.h
+++ b/Engine/Plugins/Runtime/CommonUI/Source/CommonUI/Public/Input/CommonUIActionRouterBase.h
@@ -119,6 +119,9 @@ public:
 	UE_API void FlushInput();
 
 	UE_API bool IsWidgetInActiveRoot(const UCommonActivatableWidget* Widget) const;
+	UE_API bool IsWidgetInLeafmostNodeHierarchy(const UWidget& Widget) const;
+
+	UE_API UCommonActivatableWidget* GetLeafmostActivatableWidget() const;
 
 	/** 
 	 * Sets Input Config 
