[Cert][Handheld Compatibility] UI legibility - The title displays text below the required pixel size in multiple areas when the device is set to the ?Recommended? display resolution

Override font size when small text scaling and compact mode are enabled

#Synced-CL 50957818

#Synced-CL 50957818

diff --git a/Engine/Plugins/Runtime/CommonUI/Source/CommonUI/Public/CommonUIUtils.h b/Engine/Plugins/Runtime/CommonUI/Source/CommonUI/Public/CommonUIUtils.h
--- a/Engine/Plugins/Runtime/CommonUI/Source/CommonUI/Public/CommonUIUtils.h
+++ b/Engine/Plugins/Runtime/CommonUI/Source/CommonUI/Public/CommonUIUtils.h
@@ -11,7 +11,7 @@ class IWidgetCompilerLog;
 
 namespace CommonUIUtils
 {
-	extern float UISmallTextScaling;
+	COMMONUI_API extern float UISmallTextScaling;
 	extern float UISmallTextScalingThresholdMax;
 
 	enum class ECollisionPolicy
