Rename RequestRefreshFocusIfDescendant to TryRequestRefreshFocus and provide a descriptive tooltip

#changelist validated
#tests Frontend

#Synced-CL 50971528

diff --git a/Engine/Plugins/Runtime/CommonUI/Source/CommonUI/Private/CommonUILibrary.cpp b/Engine/Plugins/Runtime/CommonUI/Source/CommonUI/Private/CommonUILibrary.cpp
--- a/Engine/Plugins/Runtime/CommonUI/Source/CommonUI/Private/CommonUILibrary.cpp
+++ b/Engine/Plugins/Runtime/CommonUI/Source/CommonUI/Private/CommonUILibrary.cpp
@@ -21,11 +21,11 @@ UCommonUILibrary::UCommonUILibrary(const FObjectInitializer& ObjectInitializer)
 {
 }
 
-bool UCommonUILibrary::RequestRefreshFocusIfLeafmostDescendant(UWidget* ContextWidget)
+bool UCommonUILibrary::RefreshFocusIfLeafmostDescendant(UWidget* ContextWidget)
 {
 	if (!ContextWidget)
 	{
-		UE_LOG(LogCommonUI, Warning, TEXT("RequestRefreshFocusIfLeafmostDescendant was called on/from a null widget."));
+		UE_LOG(LogCommonUI, Warning, TEXT("RefreshFocusIfLeafmostDescendant was called on/from a null widget."));
 		return false;
 	}
 	
@@ -40,6 +40,10 @@ bool UCommonUILibrary::RequestRefreshFocusIfLeafmostDescendant(UWidget* ContextW
 				ActivatableWidget->RequestRefreshFocus();
 				return true;
 			}
+			else
+			{
+				UE_LOG(LogCommonUI, Warning, TEXT("RefreshFocusIfLeafmostDescendant was called on/from a widget(%s) that was not a descendant of the leaf-most widget."), *ContextWidget->GetName());
+			}
 		}
 	}
 	
diff --git a/Engine/Plugins/Runtime/CommonUI/Source/CommonUI/Public/CommonUILibrary.h b/Engine/Plugins/Runtime/CommonUI/Source/CommonUI/Public/CommonUILibrary.h
--- a/Engine/Plugins/Runtime/CommonUI/Source/CommonUI/Public/CommonUILibrary.h
+++ b/Engine/Plugins/Runtime/CommonUI/Source/CommonUI/Public/CommonUILibrary.h
@@ -19,10 +19,15 @@ class UCommonUILibrary : public UBlueprintFunctionLibrary
 public:
 
 	/**
-	 * Allow any widget that is a descendant of the leaf-most widget to initiate a refresh focus request
+	 * Request a focus refresh from any widget that is a descendant of the leaf-most activated widget.
+	 * This should be used to handle asynchronous requests to route focus from outside the leaf-most activatable widget. 
+	 * A very common use case is a visualization or asynchronous widget creation from a plugin-injected widget 
+	 * or a widget otherwise unavailable or cumbersome to access from the leaf-most activatable widget.
+	 * 
+	 * @return Whether or not the request for focus was made
 	 */
 	UFUNCTION(BlueprintCallable, Category = "CommonUI|Focus", meta=(DefaultToSelf="ContextWidget", WorldContext="ContextWidget"))
-	static UE_API bool RequestRefreshFocusIfLeafmostDescendant(UWidget* ContextWidget);
+	static UE_API bool RefreshFocusIfLeafmostDescendant(UWidget* ContextWidget);
 
 	/**
 	 * Finds the first parent widget of the given type and returns it, or null if no parent could be found.
