- Support widget animations in editor tools by ticking CommonUI widgets (return a valid World to allow UUserWidget::UpdateCanTick to proceed)
- Fix for corrupted UI materials. The default window viewport UI scale is set to -1 but needs to be positive for the shader to render correctly.

#rnx
#tests editor

#Synced-CL 51069225

#Synced-CL 51069225

diff --git a/Engine/Plugins/Runtime/CommonUI/Source/CommonUI/Private/CommonUserWidget.cpp b/Engine/Plugins/Runtime/CommonUI/Source/CommonUI/Private/CommonUserWidget.cpp
--- a/Engine/Plugins/Runtime/CommonUI/Source/CommonUI/Private/CommonUserWidget.cpp
+++ b/Engine/Plugins/Runtime/CommonUI/Source/CommonUI/Private/CommonUserWidget.cpp
@@ -2,6 +2,10 @@
 
 #include "CommonUserWidget.h"
 
+#if WITH_EDITOR
+#include "Editor.h"
+#endif
+
 #include "Engine/GameInstance.h"
 #include "CommonInputSubsystem.h"
 #include "CommonUISubsystemBase.h"
@@ -212,4 +216,16 @@ const FText UCommonUserWidget::GetPaletteCategory()
 {
 	return CommonWidgetPaletteCategories::Default;
 }
+
+UWorld* UCommonUserWidget::GetWorld() const
+{
+	UWorld* World = Super::GetWorld();
+
+	if (!World && (GEditor && !GEditor->PlayWorld) && !GIsPlayInEditorWorld)
+	{
+		return GWorld;
+	}
+
+	return World;
+}
 #endif // WITH_EDITOR
diff --git a/Engine/Plugins/Runtime/CommonUI/Source/CommonUI/Public/CommonUserWidget.h b/Engine/Plugins/Runtime/CommonUI/Source/CommonUI/Public/CommonUserWidget.h
--- a/Engine/Plugins/Runtime/CommonUI/Source/CommonUI/Public/CommonUserWidget.h
+++ b/Engine/Plugins/Runtime/CommonUI/Source/CommonUI/Public/CommonUserWidget.h
@@ -45,6 +45,10 @@ public:
 	UFUNCTION(BlueprintCallable, Category = CommonUserWidget)
 	UE_API void UnregisterScrollRecipientExternal(const UWidget* AnalogScrollRecipient);
 
+#if WITH_EDITOR
+	UE_API virtual class UWorld* GetWorld() const override;
+#endif
+
 public:
 
 	const TArray<FUIActionBindingHandle>& GetActionBindings() const { return ActionBindings; }
