LogUIActionRouter - include SupportsActivationFocus in CommonUI.DumpActivatableTree

#changelist validated

#Synced-CL 51292685

diff --git a/Engine/Plugins/Runtime/CommonUI/Source/CommonUI/Private/Input/UIActionRouterTypes.cpp b/Engine/Plugins/Runtime/CommonUI/Source/CommonUI/Private/Input/UIActionRouterTypes.cpp
--- a/Engine/Plugins/Runtime/CommonUI/Source/CommonUI/Private/Input/UIActionRouterTypes.cpp
+++ b/Engine/Plugins/Runtime/CommonUI/Source/CommonUI/Private/Input/UIActionRouterTypes.cpp
@@ -1201,12 +1201,21 @@ void FActivatableTreeNode::DebugDumpRecursive(FString& OutputStr, int32 Depth, b
 
 	if (RepresentedWidget.IsValid())
 	{
-		OutputStr.Appendf(TEXT("\n%s%s: IsActivated? [%s]. LayerId=[%d]. [%d] Normal Bindings."),
+		TSharedPtr<SWidget> CachedWidget = RepresentedWidget->GetCachedWidget();
+
+		const int32 LayerId = CachedWidget.IsValid()
+			? CachedWidget->GetPersistentState().LayerId
+			: INDEX_NONE;
+
+		OutputStr.Appendf(
+			TEXT("\n%s%s: IsActivated? [%s]. LayerId=[%d]. [%d] Normal Bindings. SupportsActivationFocus [%s]"),
 			FCString::Spc(DepthSpacing),
 			*RepresentedWidget->GetName(),
 			*LexToString(RepresentedWidget->IsActivated()),
-			RepresentedWidget->GetCachedWidget()->GetPersistentState().LayerId,
-			ActionBindings.Num());
+			LayerId,
+			ActionBindings.Num(),
+			*LexToString(RepresentedWidget->SupportsActivationFocus())
+		);
 	}
 	else
 	{
