CommonButtonGroupBase - Add OnSelectedButtonLostSelection delegate to UCommonButtonGroupBase

#changelist validated

#Synced-CL 51292685

diff --git a/Engine/Plugins/Runtime/CommonUI/Source/CommonUI/Private/Groups/CommonButtonGroupBase.cpp b/Engine/Plugins/Runtime/CommonUI/Source/CommonUI/Private/Groups/CommonButtonGroupBase.cpp
--- a/Engine/Plugins/Runtime/CommonUI/Source/CommonUI/Private/Groups/CommonButtonGroupBase.cpp
+++ b/Engine/Plugins/Runtime/CommonUI/Source/CommonUI/Private/Groups/CommonButtonGroupBase.cpp
@@ -294,9 +294,15 @@ void UCommonButtonGroupBase::OnRemoveAll()
 
 void UCommonButtonGroupBase::OnSelectionStateChangedBase( UCommonButtonBase* BaseButton, bool bIsSelected )
 {
+	UCommonButtonBase* PreviouslySelectedButton = GetButtonBaseAtIndex(SelectedButtonIndex);
+
 	if (bIsSelected)
 	{
+		// Reset index to avoid stale leakage
 		SelectedButtonIndex = INDEX_NONE;
+
+		bool bFoundButton = false;
+
 		for (int32 ButtonIndex = 0; ButtonIndex < Buttons.Num(); ButtonIndex++)
 		{
 			UCommonButtonBase* Button = GetButtonBaseAtIndex(ButtonIndex);
@@ -305,6 +311,7 @@ void UCommonButtonGroupBase::OnSelectionStateChangedBase( UCommonButtonBase* Bas
 				if (Button == BaseButton)
 				{
 					SelectedButtonIndex = ButtonIndex;
+					bFoundButton = true;
 				}
 				else if (Button->GetSelected())
 				{
@@ -314,8 +321,15 @@ void UCommonButtonGroupBase::OnSelectionStateChangedBase( UCommonButtonBase* Bas
 			}
 		}
 
+		// If the button isn't in the group, do not broadcast a stale state
+		if (!bFoundButton)
+		{
+			return;
+		}
+
 		NativeOnSelectedButtonBaseChanged.Broadcast(BaseButton, SelectedButtonIndex);
-		OnSelectedButtonBaseChanged.Broadcast(BaseButton, SelectedButtonIndex);
+		OnSelectedButtonBaseChanged.Broadcast(BaseButton, SelectedButtonIndex);		
+		OnSelectedButtonLostSelection.Broadcast(PreviouslySelectedButton);
 	}
 	else
 	{
@@ -327,6 +341,9 @@ void UCommonButtonGroupBase::OnSelectionStateChangedBase( UCommonButtonBase* Bas
 			}
 		}
 
+		PreviouslySelectedButton = GetButtonBaseAtIndex(SelectedButtonIndex);
+		OnSelectedButtonLostSelection.Broadcast(PreviouslySelectedButton);
+
 		SelectedButtonIndex = INDEX_NONE;
 
 		NativeOnSelectionCleared.Broadcast();
diff --git a/Engine/Plugins/Runtime/CommonUI/Source/CommonUI/Public/Groups/CommonButtonGroupBase.h b/Engine/Plugins/Runtime/CommonUI/Source/CommonUI/Public/Groups/CommonButtonGroupBase.h
--- a/Engine/Plugins/Runtime/CommonUI/Source/CommonUI/Public/Groups/CommonButtonGroupBase.h
+++ b/Engine/Plugins/Runtime/CommonUI/Source/CommonUI/Public/Groups/CommonButtonGroupBase.h
@@ -9,7 +9,8 @@
 
 #define UE_API COMMONUI_API
 
-DECLARE_DYNAMIC_MULTICAST_DELEGATE_TwoParams(FSimpleButtonBaseGroupDelegate, UCommonButtonBase*, AssociatedButton, int32, ButtonIndex);
+DECLARE_DYNAMIC_MULTICAST_DELEGATE_TwoParams(FSimpleButtonBaseGroupDelegate, UCommonButtonBase*, AssociatedButton, int32, ButtonIndex );
+DECLARE_DYNAMIC_MULTICAST_DELEGATE_OneParam(FButtonLostSelectionDelegate, UCommonButtonBase*, PreviouslySelectedButton);
 DECLARE_DYNAMIC_MULTICAST_DELEGATE( FOnSelectionCleared );
 
 /** 
@@ -89,6 +90,9 @@ public:
 
 	DECLARE_EVENT_TwoParams(UCommonButtonGroupBase, FNativeSimpleButtonBaseGroupDelegate, UCommonButtonBase*, int32);
 
+	UPROPERTY(BlueprintAssignable)
+	FButtonLostSelectionDelegate OnSelectedButtonLostSelection;
+
 	UPROPERTY(BlueprintAssignable)
 	FSimpleButtonBaseGroupDelegate OnSelectedButtonBaseChanged;
 	FNativeSimpleButtonBaseGroupDelegate NativeOnSelectedButtonBaseChanged;
