Changing `FocusLeafmostNode` Behavior:
- Checking if Desired and AutoRestore Widgets support keyboard focus and that they are descendants of the leaf widget
- Ignoring Currently Focus Widget if there is a PendingFocusRecipient

#Synced-CL 51318692

#Synced-CL 51318692

#Synced-CL 51318692

diff --git a/Engine/Plugins/Runtime/CommonUI/Source/CommonUI/Private/Input/UIActionRouterTypes.cpp b/Engine/Plugins/Runtime/CommonUI/Source/CommonUI/Private/Input/UIActionRouterTypes.cpp
--- a/Engine/Plugins/Runtime/CommonUI/Source/CommonUI/Private/Input/UIActionRouterTypes.cpp
+++ b/Engine/Plugins/Runtime/CommonUI/Source/CommonUI/Private/Input/UIActionRouterTypes.cpp
@@ -34,6 +34,20 @@ static const TAutoConsoleVariable<bool> CVarSyncCursorToFocusedWidgetOnInputMeth
 	TEXT("Whether or not to move the cursor to the currently focused widget when input method changes to gamepad (fix for losing focus when switching between touch and gamepad inputs)."),
 	ECVF_Default);
 
+static const TAutoConsoleVariable<bool> CVarCheckKeyboardFocusAndParentage(
+	TEXT("CommonUI.CheckKeyboardFocusAndParentage"),
+	false,
+	TEXT("If enabled, we will check for keyboard focus and parentage before setting focus"));
+
+static const TAutoConsoleVariable<bool> CVarFallbackToDesiredOnAutoRestoreFailure(
+	TEXT("CommonUI.FallbackToDesiredOnAutoRestoreFailure"),
+	false,
+	TEXT("If enabled, we set focus to Desired Widget if Auto Restore Focus Fails"));
+
+static const TAutoConsoleVariable<bool> CVarDisallowUserFocusedWidgetForPendingFocusRecipient(
+	TEXT("CommonUI.DisallowUserFocusedWidgetForPendingFocusRecipient"),
+	false,
+	TEXT("If enabled, we only check the currently focused widget when there is no pending focus recipient"));
 
 const TCHAR* InputEventToString(EInputEvent InputEvent)
 {
@@ -1726,29 +1740,78 @@ void FActivatableTreeRoot::FocusLeafmostNode()
 	}
 #endif
 
-	bool bShouldCancelDelayedFocusOperation = false;
+	TSharedPtr<SWidget> AutoRestoreTarget = LeafWidget->AutoRestoresFocus() ? PinnedLeafmostNode->GetFocusFallbackTarget() : nullptr;
+	UWidget* DesiredTarget = LeafWidget->GetDesiredFocusTarget();
 
-	if (TSharedPtr<SWidget> AutoRestoreTarget = LeafWidget->AutoRestoresFocus() ? PinnedLeafmostNode->GetFocusFallbackTarget() : nullptr)
+	const bool bCheckForKeyboardAndParentage = CVarCheckKeyboardFocusAndParentage.GetValueOnAnyThread();
+	if (bCheckForKeyboardAndParentage)
 	{
-		UE_LOG(LogUIActionRouter, Display, TEXT("[User %d] Focused AutoRestoreTarget %s"), OwnerSlateId, *FReflectionMetaData::GetWidgetDebugInfo(AutoRestoreTarget.Get()));
-		
-		bShouldCancelDelayedFocusOperation = FSlateApplication::Get().SetUserFocus(OwnerSlateId, AutoRestoreTarget);
+		if (AutoRestoreTarget.IsValid())
+		{
+			const bool bSupportsKeyboard = AutoRestoreTarget->SupportsKeyboardFocus();
+			const bool bParentOfLeaf = PinnedLeafmostNode->IsExclusiveParentOfWidget(AutoRestoreTarget);
+			if (!bSupportsKeyboard || !bParentOfLeaf)
+			{
+				AutoRestoreTarget.Reset();
+			}
+		}
+
+		if (DesiredTarget)
+		{
+			const bool bSupportsKeyboard = DesiredTarget->SupportsKeyboardFocus();
+			const bool bParentOfLeaf = PinnedLeafmostNode->IsExclusiveParentOfWidget(DesiredTarget->GetCachedWidget());
+			if (!bSupportsKeyboard || !bParentOfLeaf)
+			{
+				DesiredTarget = nullptr;
+			}
+		}
 	}
-	else if (UWidget* DesiredTarget = LeafWidget->GetDesiredFocusTarget())
+
+	const bool bSuccessfullyAutoRestoreTarget = AutoRestoreTarget.IsValid() && FSlateApplication::Get().SetUserFocus(OwnerSlateId, AutoRestoreTarget);
+	bool bShouldCancelDelayedFocusOperation = bSuccessfullyAutoRestoreTarget;
+	if (bSuccessfullyAutoRestoreTarget)
+	{
+		UE_LOG(LogUIActionRouter, Display, TEXT("[User %d] Focused AutoRestoreTarget %s"), OwnerSlateId, *FReflectionMetaData::GetWidgetDebugInfo(AutoRestoreTarget.Get()));	
+	}
+	
+	const bool bFallbackToDesiredOnAutoRestoreFailure = CVarFallbackToDesiredOnAutoRestoreFailure.GetValueOnAnyThread();
+	const bool bAttemptToFocusDesiredTarget = DesiredTarget && (
+			(!bFallbackToDesiredOnAutoRestoreFailure && !AutoRestoreTarget.IsValid()) || // Legacy Behavior
+			(bFallbackToDesiredOnAutoRestoreFailure && !bSuccessfullyAutoRestoreTarget) // New Behavior
+			);
+
+	if (bAttemptToFocusDesiredTarget)
 	{
 		UE_LOG(LogUIActionRouter, Display, TEXT("[User %d] Focused desired target %s"), OwnerSlateId, *DesiredTarget->GetName());
 		DesiredTarget->SetFocus();
 	}
-	else
+	
+	if (!bSuccessfullyAutoRestoreTarget && !bAttemptToFocusDesiredTarget)
 	{
+		UE_LOG(LogUIActionRouter, Error, TEXT("[User %d] Failed to Auto Restore or Focus Desired Target for %s"), GetActionRouter().GetLocalPlayerIndex(), *LeafWidget->GetName());
+
 		// Temporary support for existing stuff that manually sets focus in Activate. If they just requested it this frame, 
 		// we won't actually be updating the focus until the LocalPlayer's FReply is processed. So if we've got a new focus target already set,
 		//	evaluate that. Otherwise evaluate the currently focused widget. Going forward, all widgets should be built to rely fully on GetDesiredFocusTarget().
 		TSharedPtr<SWidget> FocusedWidget = FSlateApplication::Get().GetUserFocusedWidget(OwnerSlateId);
 		TSharedPtr<SWidget> PendingFocusRecipient = GetActionRouter().GetLocalPlayerChecked()->GetSlateOperations().GetUserFocusRecepient();
-		if ((PendingFocusRecipient && PinnedLeafmostNode->IsExclusiveParentOfWidget(PendingFocusRecipient)) || PinnedLeafmostNode->IsExclusiveParentOfWidget(FocusedWidget))
+
+		const bool bIgnoreFocusedWidget = CVarDisallowUserFocusedWidgetForPendingFocusRecipient.GetValueOnAnyThread();
+		const bool bCheckParentOfFocusedWidget = !PendingFocusRecipient.IsValid() || !bIgnoreFocusedWidget;
+
+		const bool bIsParentOfPendingFocusRecipient = PinnedLeafmostNode->IsExclusiveParentOfWidget(PendingFocusRecipient);
+		const bool bIsParentOfFocusedWidget = PinnedLeafmostNode->IsExclusiveParentOfWidget(FocusedWidget);
+
+		if (bIsParentOfPendingFocusRecipient || (bCheckParentOfFocusedWidget && bIsParentOfFocusedWidget))
 		{
-			UE_LOG(LogUIActionRouter, Display, TEXT("[User %d] Leaf-most node [%s] did not set focus through desired methods, but the currently focused widget is acceptable. Doing nothing, but this widget should be updated."), GetActionRouter().GetLocalPlayerIndex(),  *LeafWidget->GetName());
+			if (bIgnoreFocusedWidget)
+			{
+				UE_LOG(LogUIActionRouter, Warning, TEXT("[User %d] Leaf-most node [%s] did not set focus through `GetDesiredFocusTarget`. Doing nothing. this widget should be updated."), GetActionRouter().GetLocalPlayerIndex(),  *LeafWidget->GetName());
+			}
+			else
+			{
+				UE_LOG(LogUIActionRouter, Display, TEXT("[User %d] Leaf-most node [%s] did not set focus through desired methods, but the currently focused widget is acceptable. Doing nothing, but this widget should be updated."), GetActionRouter().GetLocalPlayerIndex(),  *LeafWidget->GetName());
+			}
 		}
 		else
 		{
diff --git a/Engine/Source/Editor/GameProjectGeneration/Private/GameProjectUtils.cpp b/Engine/Source/Editor/GameProjectGeneration/Private/GameProjectUtils.cpp
--- a/Engine/Source/Editor/GameProjectGeneration/Private/GameProjectUtils.cpp
+++ b/Engine/Source/Editor/GameProjectGeneration/Private/GameProjectUtils.cpp
@@ -258,6 +258,24 @@ namespace
 				TEXT("TriggerClick"),
 				true /* ShouldReplaceExistingValue */);
 		}
+
+		ConfigValues.Emplace(TEXT("DefaultGame.ini"),
+							 TEXT("ConsoleVariables"),
+							 TEXT("CommonUI.CheckKeyboardFocusAndParentage"),
+							 TEXT("1"),
+							 true /* ShouldReplaceExistingValue */);
+
+		ConfigValues.Emplace(TEXT("DefaultGame.ini"),
+							 TEXT("ConsoleVariables"),
+							 TEXT("CommonUI.DisallowUserFocusedWidgetForPendingFocusRecipient"),
+							 TEXT("1"),
+							 true /* ShouldReplaceExistingValue */);
+
+		ConfigValues.Emplace(TEXT("DefaultGame.ini"),
+							 TEXT("ConsoleVariables"),
+							 TEXT("CommonUI.FallbackToDesiredOnAutoRestoreFailure"),
+							 TEXT("1"),
+							 true /* ShouldReplaceExistingValue */);
 	}
 } // namespace <>
 
@@ -1000,7 +1018,7 @@ void GameProjectUtils::CheckForOutOfDateGameProjectFile()
 				}
 			}; 
 
-			// If the user has specified to supress this pop up, then just early out and exit	
+			// If the user has specified to suppress this pop up, then just early out and exit	
 			if (FSuppressDialogOptions::ShouldSuppressModal())
 			{
 				return;
diff --git a/Engine/Source/Runtime/UMG/Private/Components/Widget.cpp b/Engine/Source/Runtime/UMG/Private/Components/Widget.cpp
--- a/Engine/Source/Runtime/UMG/Private/Components/Widget.cpp
+++ b/Engine/Source/Runtime/UMG/Private/Components/Widget.cpp
@@ -608,45 +608,31 @@ PRAGMA_ENABLE_DEPRECATION_WARNINGS
 bool UWidget::IsHovered() const
 {
 	TSharedPtr<SWidget> SafeWidget = GetCachedWidget();
-	if (SafeWidget.IsValid())
-	{
-		return SafeWidget->IsHovered();
-	}
-
-	return false;
+	return SafeWidget.IsValid() && SafeWidget->IsHovered();
 }
 
 bool UWidget::HasKeyboardFocus() const
 {
 	TSharedPtr<SWidget> SafeWidget = GetCachedWidget();
-	if (SafeWidget.IsValid())
-	{
-		return SafeWidget->HasKeyboardFocus();
-	}
+	return SafeWidget.IsValid() && SafeWidget->HasKeyboardFocus();
+}
 
-	return false;
+bool UWidget::SupportsKeyboardFocus() const
+{
+	TSharedPtr<SWidget> SafeWidget = GetCachedWidget();
+	return SafeWidget.IsValid() && SafeWidget->SupportsKeyboardFocus();
 }
 
 bool UWidget::HasMouseCapture() const
 {
 	TSharedPtr<SWidget> SafeWidget = GetCachedWidget();
-	if (SafeWidget.IsValid())
-	{
-		return SafeWidget->HasMouseCapture();
-	}
-
-	return false;
+	return SafeWidget.IsValid() && SafeWidget->HasMouseCapture();
 }
 
 bool UWidget::HasMouseCaptureByUser(int32 UserIndex, int32 PointerIndex) const
 {
 	TSharedPtr<SWidget> SafeWidget = GetCachedWidget();
-	if (SafeWidget.IsValid())
-	{
-		return SafeWidget->HasMouseCaptureByUser(UserIndex, PointerIndex >= 0 ? PointerIndex : TOptional<int32>());
-	}
-
-	return false;
+	return SafeWidget.IsValid() && SafeWidget->HasMouseCaptureByUser(UserIndex, PointerIndex >= 0 ? PointerIndex : TOptional<int32>());
 }
 
 void UWidget::SetKeyboardFocus()
diff --git a/Engine/Source/Runtime/UMG/Public/Components/Widget.h b/Engine/Source/Runtime/UMG/Public/Components/Widget.h
--- a/Engine/Source/Runtime/UMG/Public/Components/Widget.h
+++ b/Engine/Source/Runtime/UMG/Public/Components/Widget.h
@@ -631,6 +631,14 @@ public:
 	UFUNCTION(BlueprintCallable, Category="Widget")
 	UMG_API bool HasKeyboardFocus() const;
 
+	/**
+	* Checks to see if this widget supports keyboard focus
+	*
+	* @return  True if this widget supports keyboard focus
+	*/
+	UFUNCTION(BlueprintCallable, Category="Widget")
+	UMG_API bool SupportsKeyboardFocus() const;
+
 	/**
 	 * Checks to see if this widget is the current mouse captor
 	 * @return  True if this widget has captured the mouse
