Downgrading new Errrors and Warning messages to Display.

#Synced-CL 51402277

diff --git a/Engine/Plugins/Runtime/CommonUI/Source/CommonUI/Private/Input/UIActionRouterTypes.cpp b/Engine/Plugins/Runtime/CommonUI/Source/CommonUI/Private/Input/UIActionRouterTypes.cpp
--- a/Engine/Plugins/Runtime/CommonUI/Source/CommonUI/Private/Input/UIActionRouterTypes.cpp
+++ b/Engine/Plugins/Runtime/CommonUI/Source/CommonUI/Private/Input/UIActionRouterTypes.cpp
@@ -1788,7 +1788,7 @@ void FActivatableTreeRoot::FocusLeafmostNode()
 	
 	if (!bSuccessfullyAutoRestoreTarget && !bAttemptToFocusDesiredTarget)
 	{
-		UE_LOG(LogUIActionRouter, Error, TEXT("[User %d] Failed to Auto Restore or Focus Desired Target for %s"), GetActionRouter().GetLocalPlayerIndex(), *LeafWidget->GetName());
+		UE_LOG(LogUIActionRouter, Display, TEXT("[User %d] Failed to Auto Restore or Focus Desired Target for %s"), GetActionRouter().GetLocalPlayerIndex(), *LeafWidget->GetName());
 
 		// Temporary support for existing stuff that manually sets focus in Activate. If they just requested it this frame, 
 		// we won't actually be updating the focus until the LocalPlayer's FReply is processed. So if we've got a new focus target already set,
@@ -1806,7 +1806,7 @@ void FActivatableTreeRoot::FocusLeafmostNode()
 		{
 			if (bIgnoreFocusedWidget)
 			{
-				UE_LOG(LogUIActionRouter, Warning, TEXT("[User %d] Leaf-most node [%s] did not set focus through `GetDesiredFocusTarget`. Doing nothing. this widget should be updated."), GetActionRouter().GetLocalPlayerIndex(),  *LeafWidget->GetName());
+				UE_LOG(LogUIActionRouter, Display, TEXT("[User %d] Leaf-most node [%s] did not set focus through `GetDesiredFocusTarget`. Doing nothing. this widget should be updated."), GetActionRouter().GetLocalPlayerIndex(),  *LeafWidget->GetName());
 			}
 			else
 			{
