CommonUI: Made UpdateInputActionWidgetVisibility() function virtual to be able override it and use other input action widgets.

#rnx
#Synced-CL 51368439

diff --git a/Engine/Plugins/Runtime/CommonUI/Source/CommonUI/Public/CommonButtonBase.h b/Engine/Plugins/Runtime/CommonUI/Source/CommonUI/Public/CommonButtonBase.h
--- a/Engine/Plugins/Runtime/CommonUI/Source/CommonUI/Public/CommonButtonBase.h
+++ b/Engine/Plugins/Runtime/CommonUI/Source/CommonUI/Public/CommonButtonBase.h
@@ -803,7 +803,7 @@ protected:
 	COMMONUI_API virtual void NativeOnMouseEnter( const FGeometry& InGeometry, const FPointerEvent& InMouseEvent ) override;
 	COMMONUI_API virtual void NativeOnMouseLeave( const FPointerEvent& InMouseEvent ) override;
 
-	COMMONUI_API void UpdateInputActionWidgetVisibility();
+	COMMONUI_API virtual void UpdateInputActionWidgetVisibility();
 
 	/** The minimum width of the button (only used if greater than the style's minimum) */
 	UPROPERTY(EditAnywhere, BlueprintReadOnly, Category = Layout, meta = (ClampMin = "0"))
