[UIActionRouter] Log the normal processed inputs.
#rnx
#Synced-CL 51535417

diff --git a/Engine/Plugins/Runtime/CommonUI/Source/CommonUI/Private/Input/UIActionRouterTypes.cpp b/Engine/Plugins/Runtime/CommonUI/Source/CommonUI/Private/Input/UIActionRouterTypes.cpp
--- a/Engine/Plugins/Runtime/CommonUI/Source/CommonUI/Private/Input/UIActionRouterTypes.cpp
+++ b/Engine/Plugins/Runtime/CommonUI/Source/CommonUI/Private/Input/UIActionRouterTypes.cpp
@@ -777,6 +777,8 @@ bool FActionRouterBindingCollection::ProcessNormalInput(ECommonInputMode ActiveI
 					const bool bIsDisplayedPersistentAction = Binding->bIsPersistent && Binding->bDisplayInActionBar;
 					if (InKey.IsSameResolvedKey(Key) && Binding->InputEvent == InputEvent && (bIsDisplayedPersistentAction || IsWidgetReachableForInput(Binding->BoundWidget.Get())))
 					{
+						UE_LOGF(LogUIActionRouter, VeryVerbose, "Process normal input: %ls (key: %ls)", *Binding->ActionName.ToString(), *Key.ToString());
+						
 						// Just in case this was in the middle of a hold process with a different key, reset now
 						Binding->CancelHold();
 
