[CommonUI] Add support to register UI Action from BP

- Added multiple function inside the CommonUserWidget to support input registration the same way we do it in code.
- This could be expanded later to add more events but this covers most sue case from tech designer for now

#Synced-CL 52123706

diff --git a/Engine/Plugins/Runtime/CommonUI/Source/CommonUI/Private/CommonUserWidget.cpp b/Engine/Plugins/Runtime/CommonUI/Source/CommonUI/Private/CommonUserWidget.cpp
--- a/Engine/Plugins/Runtime/CommonUI/Source/CommonUI/Private/CommonUserWidget.cpp
+++ b/Engine/Plugins/Runtime/CommonUI/Source/CommonUI/Private/CommonUserWidget.cpp
@@ -12,6 +12,8 @@
 #include "CommonWidgetPaletteCategories.h"
 #include "Input/CommonUIActionRouterBase.h"
 #include "Input/CommonUIInputTypes.h"
+#include "InputMappingContext.h"
+
 
 #include UE_INLINE_GENERATED_CPP_BY_NAME(CommonUserWidget)
 
@@ -101,6 +103,49 @@ FUIActionBindingHandle UCommonUserWidget::RegisterUIActionBinding(const FBindUIA
 	return FUIActionBindingHandle();
 }
 
+FUIActionBindingHandle UCommonUserWidget::RegisterUIAction(const UInputAction* EnhancedInputAction, bool ShouldDisplayInActionBar /* = true */)
+{
+	// Empty Callback are those are meant to be used as EnhancedInputAction. If we need we could add a callback in BP with the InputAction later.
+	return RegisterUIActionBinding(
+		FBindUIActionArgs(EnhancedInputAction, ShouldDisplayInActionBar, FSimpleDelegate::CreateLambda([]{})));
+}
+
+TArray<FUIActionBindingHandle> UCommonUserWidget::RegisterUIActionsFromMappingContext(const UInputMappingContext* MappingContext
+	, bool ShouldDisplayInActionBar)
+{
+	if (MappingContext == nullptr)
+	{
+		return {};
+	}
+
+	TArray<FUIActionBindingHandle> ReturnActionBindings;
+	for (const FEnhancedActionKeyMapping& Mapping : MappingContext->GetMappings())
+	{
+		ReturnActionBindings.Add(RegisterUIAction(Mapping.Action, ShouldDisplayInActionBar));
+	}
+	
+	return ReturnActionBindings;
+}
+
+void UCommonUserWidget::RemoveUIAction(FUIActionBindingHandle ActionBinding)
+{
+	ActionBinding.Unregister();
+	RemoveActionBinding(ActionBinding);
+}
+
+void UCommonUserWidget::RemoveAllUIActionBinding()
+{
+	if (UCommonUIActionRouterBase* ActionRouter = UCommonUIActionRouterBase::Get(*this))
+	{
+		for (FUIActionBindingHandle& ActionBinding : ActionBindings)
+		{
+			ActionBinding.Unregister();
+			ActionRouter->RemoveBinding(ActionBinding);
+		}
+	}
+	ActionBindings.Reset();
+}
+
 void UCommonUserWidget::RemoveActionBinding(FUIActionBindingHandle ActionBinding)
 {
 	ActionBindings.Remove(ActionBinding);
diff --git a/Engine/Plugins/Runtime/CommonUI/Source/CommonUI/Public/CommonUserWidget.h b/Engine/Plugins/Runtime/CommonUI/Source/CommonUI/Public/CommonUserWidget.h
--- a/Engine/Plugins/Runtime/CommonUI/Source/CommonUI/Public/CommonUserWidget.h
+++ b/Engine/Plugins/Runtime/CommonUI/Source/CommonUI/Public/CommonUserWidget.h
@@ -9,6 +9,8 @@
 
 #define UE_API COMMONUI_API
 
+class UInputMappingContext;
+class UInputAction;
 class UCommonInputSubsystem;
 class UCommonUISubsystemBase;
 class FSlateUser;
@@ -59,6 +61,23 @@ public:
 	 * Persistent bindings are *always* listening for input while registered, while normal bindings are only listening when all of this widget's activatable parents are activated.
 	 */
 	UE_API FUIActionBindingHandle RegisterUIActionBinding(const FBindUIActionArgs& BindActionArgs);
+	
+	// Helper to register a UI Enhanced Input Action from BP, they are meant to be global UI action to work with Enhanced Input
+	// See the UCommonMappingContextMetadata set in the Input Action User Settings Section
+	UFUNCTION(BlueprintCallable, Category = CommonUserWidget)
+	FUIActionBindingHandle RegisterUIAction(const UInputAction* EnhancedInputAction, bool ShouldDisplayInActionBar = true);
+	
+	// Similar to RegisterUIAction but uses a Mapping Context 
+	UFUNCTION(BlueprintCallable, Category = CommonUserWidget)
+	TArray<FUIActionBindingHandle> RegisterUIActionsFromMappingContext(const UInputMappingContext* MappingContext, bool ShouldDisplayInActionBar = true);
+	
+	// Remove a register UI Action. Will silently fail as CommonActionRouter doesn't give Feedback
+	UFUNCTION(BlueprintCallable, Category = CommonUserWidget)
+	void RemoveUIAction(FUIActionBindingHandle ActionBinding);
+	
+	// Remove All the UI Action that were registered on this widget
+	UFUNCTION(BlueprintCallable, Category = CommonUserWidget)
+	void RemoveAllUIActionBinding();
 
 	UE_API void RemoveActionBinding(FUIActionBindingHandle ActionBinding);
 	UE_API void AddActionBinding(FUIActionBindingHandle ActionBinding);
