fix CommonActivatableWidget.h typo
#14605
#Synced-CL 52272722

diff --git a/Engine/Plugins/Runtime/CommonUI/Source/CommonUI/Public/CommonActivatableWidget.h b/Engine/Plugins/Runtime/CommonUI/Source/CommonUI/Public/CommonActivatableWidget.h
--- a/Engine/Plugins/Runtime/CommonUI/Source/CommonUI/Public/CommonActivatableWidget.h
+++ b/Engine/Plugins/Runtime/CommonUI/Source/CommonUI/Public/CommonActivatableWidget.h
@@ -226,7 +226,7 @@ protected:
 	UPROPERTY(EditAnywhere, BlueprintReadOnly, Category = "Input", meta = (InlineEditConditionToggle))
 	bool bOverrideActionDomain = false;
 	
-	/** Optional mapping context to be applied & removed on activation & deactivation respectfully. */
+	/** Optional mapping context to be applied & removed on activation & deactivation, respectively. */
 	UPROPERTY(EditAnywhere, Category="Input", meta = (EditCondition = "CommonInput.CommonInputSettings.IsEnhancedInputSupportEnabled", EditConditionHides))
 	TObjectPtr<UInputMappingContext> InputMapping;
 
