PR #14284: CommonUI: Support lazy loading atlas/sprite textures

#rnx
#Synced-CL 52347135

diff --git a/Engine/Plugins/Runtime/CommonUI/Source/CommonUI/Private/CommonLazyImage.cpp b/Engine/Plugins/Runtime/CommonUI/Source/CommonUI/Private/CommonLazyImage.cpp
--- a/Engine/Plugins/Runtime/CommonUI/Source/CommonUI/Private/CommonLazyImage.cpp
+++ b/Engine/Plugins/Runtime/CommonUI/Source/CommonUI/Private/CommonLazyImage.cpp
@@ -172,6 +172,10 @@ void UCommonLazyImage::SetBrushFromLazyDisplayAsset(const TSoftObjectPtr<UObject
 					{
 						StrongThis->SetBrushFromTexture(AsTexture, bMatchTextureSize);
 					}
+					else if (TScriptInterface<ISlateTextureAtlasInterface> AsAtlasInterface = LazyObject.Get())
+					{
+						StrongThis->SetBrushFromAtlasInterface(AsAtlasInterface, bMatchTextureSize);
+					}
 					else if (UMaterialInterface* AsMaterial = Cast<UMaterialInterface>(LazyObject.Get()))
 					{
 						StrongThis->SetBrushFromMaterial(AsMaterial);
