PR #14070: Fix FCommonAnalogCursor sometimes stays at zero coordinates when a widget with valid geometry is focused

#rnx
#Synced-CL 52347135

diff --git a/Engine/Plugins/Runtime/CommonUI/Source/CommonUI/Private/Input/CommonAnalogCursor.cpp b/Engine/Plugins/Runtime/CommonUI/Source/CommonUI/Private/Input/CommonAnalogCursor.cpp
--- a/Engine/Plugins/Runtime/CommonUI/Source/CommonUI/Private/Input/CommonAnalogCursor.cpp
+++ b/Engine/Plugins/Runtime/CommonUI/Source/CommonUI/Private/Input/CommonAnalogCursor.cpp
@@ -241,6 +241,7 @@ void FCommonAnalogCursor::Tick(const float DeltaTime, FSlateApplication& SlateAp
 
 				if (!bHasValidCursorTarget)
 				{
+					LastCursorTargetTransform = FSlateRenderTransform();
 					SetNormalizedCursorPosition(FVector2D::ZeroVector);
 				}
 			}
