PR #14049: Make inline icons show in common ui rich text blocks in shipping

#rnx
#Synced-CL 52347135

diff --git a/Engine/Plugins/Runtime/CommonUI/Source/CommonUI/Private/CommonRichTextBlock.cpp b/Engine/Plugins/Runtime/CommonUI/Source/CommonUI/Private/CommonRichTextBlock.cpp
--- a/Engine/Plugins/Runtime/CommonUI/Source/CommonUI/Private/CommonRichTextBlock.cpp
+++ b/Engine/Plugins/Runtime/CommonUI/Source/CommonUI/Private/CommonRichTextBlock.cpp
@@ -269,7 +269,7 @@ private:
 	TSharedRef<SWidget> CreateIconWidget(const FInlineIconEntry& Entry, const FTextBlockStyle& TextStyle, float IconScale, bool bTintIcon)
 	{
 #if UE_BUILD_SHIPPING
-		if (!Entry.ResourceObject)
+		if (Entry.ResourceObject.IsNull())
 		{
 			// Just be blank in shipping if there's no texture
 			return SNullWidget::NullWidget;
