PR #13953: Make UCommonTextStyle's ToTextBlockStyle and ApplyToTextBlock virtual

#rnx
#Synced-CL 52347135

diff --git a/Engine/Plugins/Runtime/CommonUI/Source/CommonUI/Public/CommonTextBlock.h b/Engine/Plugins/Runtime/CommonUI/Source/CommonUI/Public/CommonTextBlock.h
--- a/Engine/Plugins/Runtime/CommonUI/Source/CommonUI/Public/CommonTextBlock.h
+++ b/Engine/Plugins/Runtime/CommonUI/Source/CommonUI/Public/CommonTextBlock.h
@@ -86,9 +86,9 @@ public:
 	UFUNCTION(BlueprintCallable, Category = "Common Text Style|Getters")
 	UE_API void GetStrikeBrush(FSlateBrush& OutStrikeBrush) const;
 
-	UE_API void ToTextBlockStyle(FTextBlockStyle& OutTextBlockStyle) const;
+	UE_API virtual void ToTextBlockStyle(FTextBlockStyle& OutTextBlockStyle) const;
 
-	UE_API void ApplyToTextBlock(const TSharedRef<STextBlock>& TextBlock) const;
+	UE_API virtual void ApplyToTextBlock(const TSharedRef<STextBlock>& TextBlock) const;
 };
 
 /* 
