PR #13486: Fix incorrect EditCondition in FCommonAnalogCursorSettings

#rnx
#Synced-CL 52347135

diff --git a/Engine/Plugins/Runtime/CommonUI/Source/CommonUI/Public/Input/CommonUIInputSettings.h b/Engine/Plugins/Runtime/CommonUI/Source/CommonUI/Public/Input/CommonUIInputSettings.h
--- a/Engine/Plugins/Runtime/CommonUI/Source/CommonUI/Public/Input/CommonUIInputSettings.h
+++ b/Engine/Plugins/Runtime/CommonUI/Source/CommonUI/Public/Input/CommonUIInputSettings.h
@@ -78,7 +78,7 @@ public:
 	UPROPERTY(EditDefaultsOnly, Category = AnalogCursor)
 	bool bEnableCursorAcceleration = true;
 
-	UPROPERTY(EditDefaultsOnly, Category = AnalogCursor, meta = (EditCondition = bEnableAnalogCursorAcceleration))
+	UPROPERTY(EditDefaultsOnly, Category = AnalogCursor, meta = (EditCondition = bEnableCursorAcceleration))
 	float CursorAcceleration = 1500.f;
 
 	UPROPERTY(EditDefaultsOnly, Category = AnalogCursor)
