UI Hold Input Action fix.
- fix for enhanced input action with hold trigger being processed as a single press instead of hold (use CommonUI.bEnhancedInputDynamicallyResolveHoldKeysForInputAction)
- press and quick release will not trigger pressed event while hold is expected (use CommonUI.AlwaysWaitForHoldActionToComplete)

#Synced-CL 52547818

#Synced-CL 52555157

diff --git a/Engine/Plugins/Runtime/CommonUI/Source/CommonUI/Private/Input/UIActionRouterTypes.cpp b/Engine/Plugins/Runtime/CommonUI/Source/CommonUI/Private/Input/UIActionRouterTypes.cpp
--- a/Engine/Plugins/Runtime/CommonUI/Source/CommonUI/Private/Input/UIActionRouterTypes.cpp
+++ b/Engine/Plugins/Runtime/CommonUI/Source/CommonUI/Private/Input/UIActionRouterTypes.cpp
@@ -18,6 +18,7 @@
 #include "Widgets/SViewport.h"
 #include "CommonButtonBase.h"
 #include "CommonUIEditorSettings.h"
+#include "EnhancedInputSubsystems.h"
 #include "Engine/GameInstance.h"
 #include "ICommonUIModule.h"
 
@@ -49,6 +50,18 @@ static const TAutoConsoleVariable<bool> CVarDisallowUserFocusedWidgetForPendingF
 	false,
 	TEXT("If enabled, we only check the currently focused widget when there is no pending focus recipient"));
 
+static bool bEnhancedInputDynamicallyResolveHoldKeysForInputAction = false;
+static FAutoConsoleVariableRef CVarEnhancedInputDynamicallyResolveHoldKeysForInputAction(
+	TEXT("CommonUI.bEnhancedInputDynamicallyResolveHoldKeysForInputAction"),
+	bEnhancedInputDynamicallyResolveHoldKeysForInputAction,
+	TEXT("If enabled, triggering keys assigned to InputAction for hold triggers will be resolved dynamically during key press processing (this allows remapped keys)."));
+	
+static bool bAlwaysWaitForHoldActionToComplete = false;
+static FAutoConsoleVariableRef CVarAlwaysWaitForHoldActionToComplete(
+	TEXT("CommonUI.AlwaysWaitForHoldActionToComplete"),
+	bAlwaysWaitForHoldActionToComplete,
+	TEXT("If enabled, always wait for input hold time to reach triggering threshold before generating a press event (otherwise a press event will be triggered before threshold is reached if press and quick release key)."));
+
 const TCHAR* InputEventToString(EInputEvent InputEvent)
 {
 	switch (InputEvent)
@@ -698,6 +711,41 @@ EProcessHoldActionResult FActionRouterBindingCollection::ProcessHoldInput(ECommo
 		{
 			if (Binding->UserIndex == UserIndex && (ActiveInputMode == ECommonInputMode::All || ActiveInputMode == Binding->InputMode))
 			{
+				// Create dynamic key mappings. This also ensures that remapped keys are properly processed.
+				if (CommonUI::IsEnhancedInputSupportEnabled() && bEnhancedInputDynamicallyResolveHoldKeysForInputAction && Binding->InputAction.IsValid())
+				{
+					Binding->HoldMappings.Empty();
+					
+					if (const UInputAction* InputAction = Binding->InputAction.Get())
+					{
+						// find hold trigger and corresponding keys for this action if a hold trigger is found
+						for (const TObjectPtr<UInputTrigger>& Trigger : InputAction->Triggers)
+						{
+							if (const UInputTriggerHold* HoldTrigger = Cast<UInputTriggerHold>(Trigger))
+							{
+								if (HoldTrigger->HoldTimeThreshold <= 0.0f)
+								{
+									continue;
+								}
+							
+								TArray<FKey> InputActionKeys;
+								CommonUI::GetEnhancedInputActionKeys(BindingHandle.GetBoundLocalPlayer(), InputAction, InputActionKeys);
+							
+								for (const FKey& InputActionKey : InputActionKeys)
+								{
+									if (InputActionKey.IsSameResolvedKey(Key))
+									{
+										Binding->HoldMappings.Emplace(FUIActionKeyMapping(InputActionKey, HoldTrigger->HoldTimeThreshold, 0.f));
+										break;
+									}
+								}
+								
+								break;
+							}
+						}
+					}
+				}
+
 				for (const FUIActionKeyMapping& HoldMapping : Binding->HoldMappings)
 				{
 					EProcessHoldActionResult ProcessResult = EProcessHoldActionResult::Unhandled;
@@ -742,10 +790,13 @@ EProcessHoldActionResult FActionRouterBindingCollection::ProcessHoldInput(ECommo
 									Binding->BeginRollback(HoldMapping.HoldRollbackTime, HoldMapping.HoldTime, Binding->Handle);
 								}
 
-								static const float PressToHoldThreshold = 0.25f;
-								if (SecondsHeld <= PressToHoldThreshold && SecondsHeld < HoldMapping.HoldTime)
+								if (!bAlwaysWaitForHoldActionToComplete)
 								{
-									ProcessResult = EProcessHoldActionResult::GeneratePress;
+									static const float PressToHoldThreshold = 0.25f;
+									if (SecondsHeld <= PressToHoldThreshold && SecondsHeld < HoldMapping.HoldTime)
+									{
+										ProcessResult = EProcessHoldActionResult::GeneratePress;
+									}
 								}
 							}
 						}
@@ -1775,7 +1826,7 @@ void FActivatableTreeRoot::FocusLeafmostNode()
 	{
 		UE_LOG(LogUIActionRouter, Display, TEXT("[User %d] Focused AutoRestoreTarget %s"), OwnerSlateId, *FReflectionMetaData::GetWidgetDebugInfo(AutoRestoreTarget.Get()));	
 	}
-	
+		
 	const bool bFallbackToDesiredOnAutoRestoreFailure = CVarFallbackToDesiredOnAutoRestoreFailure.GetValueOnAnyThread();
 	const bool bAttemptToFocusDesiredTarget = DesiredTarget && (
 			(!bFallbackToDesiredOnAutoRestoreFailure && !AutoRestoreTarget.IsValid()) || // Legacy Behavior
@@ -1811,9 +1862,9 @@ void FActivatableTreeRoot::FocusLeafmostNode()
 				UE_LOG(LogUIActionRouter, Display, TEXT("[User %d] Leaf-most node [%s] did not set focus through `GetDesiredFocusTarget`. Doing nothing. this widget should be updated."), GetActionRouter().GetLocalPlayerIndex(),  *LeafWidget->GetName());
 			}
 			else
-			{
-				UE_LOG(LogUIActionRouter, Display, TEXT("[User %d] Leaf-most node [%s] did not set focus through desired methods, but the currently focused widget is acceptable. Doing nothing, but this widget should be updated."), GetActionRouter().GetLocalPlayerIndex(),  *LeafWidget->GetName());
-			}
+		{
+			UE_LOG(LogUIActionRouter, Display, TEXT("[User %d] Leaf-most node [%s] did not set focus through desired methods, but the currently focused widget is acceptable. Doing nothing, but this widget should be updated."), GetActionRouter().GetLocalPlayerIndex(),  *LeafWidget->GetName());
+		}
 		}
 		else
 		{
